SkeinUE: Initial skeleton code commit for the Skein Source Control Provider plugin.

#rnx

diff --git a/Engine/Plugins/Developer/SkeinSourceControl/SkeinSourceControl.uplugin b/Engine/Plugins/Developer/SkeinSourceControl/SkeinSourceControl.uplugin
new file mode 100644
--- /dev/null
+++ b/Engine/Plugins/Developer/SkeinSourceControl/SkeinSourceControl.uplugin
@@ -0,0 +1,25 @@
+{
+	"FileVersion" : 3,
+	"Version" : 1,
+	"VersionName" : "1.0",
+	"FriendlyName" : "Skein",
+	"Description" : "Skein source control management",
+	"Category" : "Source Control",
+	"CreatedBy" : "Epic Games, Inc.",
+	"CreatedByURL" : "http://epicgames.com",
+	"DocsURL" : "",
+	"MarketplaceURL" : "",
+	"SupportURL" : "",
+	"EnabledByDefault" : false,
+	"CanContainContent" : false,
+	"IsBetaVersion" : true,
+	"Installed" : false,
+	"Modules" :
+	[
+		{
+			"Name" : "SkeinSourceControl",
+			"Type" : "Editor",
+			"LoadingPhase" : "Default"
+		}
+	]
+}
\ No newline at end of file
diff --git a/Engine/Plugins/Developer/SkeinSourceControl/Source/SkeinSourceControl/Private/ISkeinSourceControlWorker.h b/Engine/Plugins/Developer/SkeinSourceControl/Source/SkeinSourceControl/Private/ISkeinSourceControlWorker.h
new file mode 100644
--- /dev/null
+++ b/Engine/Plugins/Developer/SkeinSourceControl/Source/SkeinSourceControl/Private/ISkeinSourceControlWorker.h
@@ -0,0 +1,27 @@
+// Copyright Epic Games, Inc. All Rights Reserved.
+
+#pragma once
+
+#include "CoreMinimal.h"
+
+class ISkeinSourceControlWorker
+{
+public:
+	/**
+	 * Name describing the work that this worker does. Used for factory method hookup.
+	 */
+	virtual FName GetName() const = 0;
+
+	/**
+	 * Function that actually does the work. Can be executed on another thread.
+	 */
+	virtual bool Execute(class FSkeinSourceControlCommand& InCommand) = 0;
+
+	/**
+	 * Updates the state of any items after completion (if necessary). This is always executed on the main thread.
+	 * @returns true if states were updated
+	 */
+	virtual bool UpdateStates() const = 0;
+};
+
+typedef TSharedRef<ISkeinSourceControlWorker, ESPMode::ThreadSafe> FSkeinSourceControlWorkerRef;
diff --git a/Engine/Plugins/Developer/SkeinSourceControl/Source/SkeinSourceControl/Private/SSkeinSourceControlSettings.cpp b/Engine/Plugins/Developer/SkeinSourceControl/Source/SkeinSourceControl/Private/SSkeinSourceControlSettings.cpp
new file mode 100644
--- /dev/null
+++ b/Engine/Plugins/Developer/SkeinSourceControl/Source/SkeinSourceControl/Private/SSkeinSourceControlSettings.cpp
@@ -0,0 +1,65 @@
+// Copyright Epic Games, Inc. All Rights Reserved.
+
+#include "SSkeinSourceControlSettings.h"
+#include "SkeinSourceControlModule.h"
+
+#define LOCTEXT_NAMESPACE "SSkeinSourceControlSettings"
+
+void SSkeinSourceControlSettings::Construct(const FArguments& InArgs)
+{
+	ChildSlot
+	[
+		SNew(SVerticalBox)
+		// Explanation text for missing Skein CLI
+		+SVerticalBox::Slot()
+		.FillHeight(1.0f)
+		.Padding(2.0f)
+		.VAlign(VAlign_Center)
+		[
+			SNew(SHorizontalBox)
+			.Visibility(this, &SSkeinSourceControlSettings::CanUseSkeinCLI)
+			+SHorizontalBox::Slot()
+			.FillWidth(1.0f)
+			[
+				SNew(STextBlock)
+				.Text(LOCTEXT("SkeinCLINotFound", "The Skein Command Line application could not be found."))
+				.WrapTextAt(450.0f)
+				.ToolTipText(LOCTEXT("SkeinCLINotFound_Tooltip", "Without the Skein Command Line application the Unreal Editor cannot communicate to the Skein cloud server. Please make sure it's installed at the correct location."))
+			]
+		]
+		// Explanation text for missing Skein project
+		+SVerticalBox::Slot()
+		.FillHeight(1.0f)
+		.Padding(2.0f)
+		.VAlign(VAlign_Center)
+		[
+			SNew(SHorizontalBox)
+			.Visibility(this, &SSkeinSourceControlSettings::CanUseSkeinProject)
+			+SHorizontalBox::Slot()
+			.FillWidth(1.0f)
+			[
+				SNew(STextBlock)
+				.Text(LOCTEXT("SkeinProjectNotFound", "There is no Skein project initialized for this location."))
+				.WrapTextAt(450.0f)
+				.ToolTipText(LOCTEXT("SkeinProjectNotFound_Tooltip", "Skein projects are created in the Skein Web UI and then cloned to your local machine using the Skein Command Line application. Once complete, you can use the Skein Source Control plugin to manage the project assets from within Unreal Editor."))
+			]
+		]
+	];
+}
+
+EVisibility SSkeinSourceControlSettings::CanUseSkeinCLI() const
+{
+	FSkeinSourceControlModule& SkeinSourceControl = FModuleManager::LoadModuleChecked<FSkeinSourceControlModule>("SkeinSourceControl");
+	const bool bSkeinAvailable = SkeinSourceControl.GetProvider().IsAvailable();
+	return (!bSkeinAvailable) ? EVisibility::Visible : EVisibility::Collapsed;
+}
+
+EVisibility SSkeinSourceControlSettings::CanUseSkeinProject() const
+{
+	FSkeinSourceControlModule& SkeinSourceControl = FModuleManager::LoadModuleChecked<FSkeinSourceControlModule>("SkeinSourceControl");
+	const bool bSkeinAvailable = SkeinSourceControl.GetProvider().IsAvailable();
+	const bool bSkeinEnabled = SkeinSourceControl.GetProvider().IsEnabled();
+	return (bSkeinAvailable && !bSkeinEnabled) ? EVisibility::Visible : EVisibility::Collapsed;
+}
+
+#undef LOCTEXT_NAMESPACE
diff --git a/Engine/Plugins/Developer/SkeinSourceControl/Source/SkeinSourceControl/Private/SSkeinSourceControlSettings.h b/Engine/Plugins/Developer/SkeinSourceControl/Source/SkeinSourceControl/Private/SSkeinSourceControlSettings.h
new file mode 100644
--- /dev/null
+++ b/Engine/Plugins/Developer/SkeinSourceControl/Source/SkeinSourceControl/Private/SSkeinSourceControlSettings.h
@@ -0,0 +1,30 @@
+// Copyright Epic Games, Inc. All Rights Reserved.
+
+#pragma once
+
+#include "CoreMinimal.h"
+#include "Layout/Visibility.h"
+#include "Input/Reply.h"
+#include "Widgets/DeclarativeSyntaxSupport.h"
+#include "Widgets/SCompoundWidget.h"
+#include "SlateFwd.h"
+#include "ISourceControlOperation.h"
+#include "ISourceControlProvider.h"
+
+class SSkeinSourceControlSettings : public SCompoundWidget
+{
+public:
+	
+	SLATE_BEGIN_ARGS(SSkeinSourceControlSettings) {}
+	
+	SLATE_END_ARGS()
+
+public:
+
+	void Construct(const FArguments& InArgs);
+
+private:
+
+	EVisibility CanUseSkeinCLI() const;
+	EVisibility CanUseSkeinProject() const;
+};
diff --git a/Engine/Plugins/Developer/SkeinSourceControl/Source/SkeinSourceControl/Private/SkeinSourceControlCommand.cpp b/Engine/Plugins/Developer/SkeinSourceControl/Source/SkeinSourceControl/Private/SkeinSourceControlCommand.cpp
new file mode 100644
--- /dev/null
+++ b/Engine/Plugins/Developer/SkeinSourceControl/Source/SkeinSourceControl/Private/SkeinSourceControlCommand.cpp
@@ -0,0 +1,60 @@
+// Copyright Epic Games, Inc. All Rights Reserved.
+
+#include "SkeinSourceControlCommand.h"
+#include "SkeinSourceControlModule.h"
+#include "Modules/ModuleManager.h"
+
+FSkeinSourceControlCommand::FSkeinSourceControlCommand(const TSharedRef<class ISourceControlOperation, ESPMode::ThreadSafe>& InOperation, const TSharedRef<class ISkeinSourceControlWorker, ESPMode::ThreadSafe>& InWorker, const FSourceControlOperationComplete& InOperationCompleteDelegate)
+	: Operation(InOperation)
+	, Worker(InWorker)
+	, OperationCompleteDelegate(InOperationCompleteDelegate)
+	, bExecuteProcessed(0)
+	, bCommandSuccessful(false)
+	, bAutoDelete(true)
+	, Concurrency(EConcurrency::Synchronous)
+{
+	check(IsInGameThread());
+
+	// Grab the providers settings here, so we don't access them once the worker thread is launched
+	FSkeinSourceControlModule& SkeinSourceControl = FModuleManager::LoadModuleChecked<FSkeinSourceControlModule>( "SkeinSourceControl" );
+	SkeinBinaryPath = SkeinSourceControl.GetProvider().GetSkeinBinaryPath();
+	SkeinProjectRoot = SkeinSourceControl.GetProvider().GetSkeinProjectRoot();
+}
+
+bool FSkeinSourceControlCommand::DoWork()
+{
+	bCommandSuccessful = Worker->Execute(*this);
+	FPlatformAtomics::InterlockedExchange(&bExecuteProcessed, 1);
+
+	return bCommandSuccessful;
+}
+
+void FSkeinSourceControlCommand::Abandon()
+{
+	FPlatformAtomics::InterlockedExchange(&bExecuteProcessed, 1);
+}
+
+void FSkeinSourceControlCommand::DoThreadedWork()
+{
+	Concurrency = EConcurrency::Asynchronous;
+	DoWork();
+}
+
+ECommandResult::Type FSkeinSourceControlCommand::ReturnResults()
+{
+	// Save any messages that have accumulated
+	for (FString& String : InfoMessages)
+	{
+		Operation->AddInfoMessge(FText::FromString(String));
+	}
+	for (FString& String : ErrorMessages)
+	{
+		Operation->AddErrorMessge(FText::FromString(String));
+	}
+
+	// Run the completion delegate if we have one bound
+	ECommandResult::Type Result = bCommandSuccessful ? ECommandResult::Succeeded : ECommandResult::Failed;
+	OperationCompleteDelegate.ExecuteIfBound(Operation, Result);
+
+	return Result;
+}
diff --git a/Engine/Plugins/Developer/SkeinSourceControl/Source/SkeinSourceControl/Private/SkeinSourceControlCommand.h b/Engine/Plugins/Developer/SkeinSourceControl/Source/SkeinSourceControl/Private/SkeinSourceControlCommand.h
new file mode 100644
--- /dev/null
+++ b/Engine/Plugins/Developer/SkeinSourceControl/Source/SkeinSourceControl/Private/SkeinSourceControlCommand.h
@@ -0,0 +1,77 @@
+// Copyright Epic Games, Inc. All Rights Reserved.
+
+#pragma once
+
+#include "CoreMinimal.h"
+#include "ISourceControlProvider.h"
+#include "Misc/IQueuedWork.h"
+
+/**
+ * Used to execute Skein commands multi-threaded.
+ */
+class FSkeinSourceControlCommand : public IQueuedWork
+{
+public:
+
+	FSkeinSourceControlCommand(const TSharedRef<class ISourceControlOperation, ESPMode::ThreadSafe>& InOperation, const TSharedRef<class ISkeinSourceControlWorker, ESPMode::ThreadSafe>& InWorker, const FSourceControlOperationComplete& InOperationCompleteDelegate = FSourceControlOperationComplete() );
+
+	/**
+	 * This is where the real thread work is done. All work that is done for
+	 * this queued object should be done from within the call to this function.
+	 */
+	bool DoWork();
+
+	/**
+	 * Tells the queued work that it is being abandoned so that it can do
+	 * per object clean up as needed. This will only be called if it is being
+	 * abandoned before completion. NOTE: This requires the object to delete
+	 * itself using whatever heap it was allocated in.
+	 */
+	virtual void Abandon() override;
+
+	/**
+	 * This method is also used to tell the object to cleanup but not before
+	 * the object has finished it's work.
+	 */
+	virtual void DoThreadedWork() override;
+
+	/** Save any results and call any registered callbacks. */
+	ECommandResult::Type ReturnResults();
+
+public:
+	/** Path to the Skein binary */
+	FString SkeinBinaryPath;
+
+	/** Path to the root of the project */
+	FString SkeinProjectRoot;
+
+	/** Operation we want to perform - contains outward-facing parameters & results */
+	TSharedRef<class ISourceControlOperation, ESPMode::ThreadSafe> Operation;
+
+	/** The object that will actually do the work */
+	TSharedRef<class ISkeinSourceControlWorker, ESPMode::ThreadSafe> Worker;
+
+	/** Delegate to notify when this operation completes */
+	FSourceControlOperationComplete OperationCompleteDelegate;
+
+	/** If true, this command has been processed by the source control thread */
+	volatile int32 bExecuteProcessed;
+
+	/** If true, the source control command succeeded */
+	bool bCommandSuccessful;
+
+	/** If true, this command will be automatically cleaned up in Tick() */
+	bool bAutoDelete;
+
+	/** Whether we are running multi-treaded or not */
+	EConcurrency::Type Concurrency;
+
+	/** Files to perform this operation on */
+	TArray<FString> Files;
+
+	/** Info and/or warning message message storage */
+	TArray<FString> InfoMessages;
+
+	/** Potential error message storage */
+	TArray<FString> ErrorMessages;
+};
diff --git a/Engine/Plugins/Developer/SkeinSourceControl/Source/SkeinSourceControl/Private/SkeinSourceControlModule.cpp b/Engine/Plugins/Developer/SkeinSourceControl/Source/SkeinSourceControl/Private/SkeinSourceControlModule.cpp
new file mode 100644
--- /dev/null
+++ b/Engine/Plugins/Developer/SkeinSourceControl/Source/SkeinSourceControl/Private/SkeinSourceControlModule.cpp
@@ -0,0 +1,44 @@
+// Copyright Epic Games, Inc. All Rights Reserved.
+
+#include "SkeinSourceControlModule.h"
+#include "SkeinSourceControlOperations.h"
+#include "Features/IModularFeatures.h"
+
+#define LOCTEXT_NAMESPACE "SkeinSourceControl"
+
+template<typename Type>
+static TSharedRef<ISkeinSourceControlWorker, ESPMode::ThreadSafe> CreateWorker()
+{
+	return MakeShareable(new Type());
+}
+
+void FSkeinSourceControlModule::StartupModule()
+{
+	// Register our operations
+	SkeinSourceControlProvider.RegisterWorker("Connect", FGetSkeinSourceControlWorker::CreateStatic(&CreateWorker<FSkeinConnectWorker>));
+	SkeinSourceControlProvider.RegisterWorker("CheckIn", FGetSkeinSourceControlWorker::CreateStatic(&CreateWorker<FSkeinCheckInWorker>));
+	SkeinSourceControlProvider.RegisterWorker("MarkForAdd", FGetSkeinSourceControlWorker::CreateStatic(&CreateWorker<FSkeinMarkForAddWorker>));
+	SkeinSourceControlProvider.RegisterWorker("Delete", FGetSkeinSourceControlWorker::CreateStatic(&CreateWorker<FSkeinDeleteWorker>));
+	SkeinSourceControlProvider.RegisterWorker("Sync", FGetSkeinSourceControlWorker::CreateStatic(&CreateWorker<FSkeinSyncWorker>));
+	SkeinSourceControlProvider.RegisterWorker("UpdateStatus", FGetSkeinSourceControlWorker::CreateStatic(&CreateWorker<FSkeinUpdateStatusWorker>));
+
+//	SkeinSourceControlProvider.RegisterWorker("Revert", FGetSkeinSourceControlWorker::CreateStatic(&CreateWorker<FSkeinRevertWorker>));
+//	SkeinSourceControlProvider.RegisterWorker("Copy", FGetSkeinSourceControlWorker::CreateStatic(&CreateWorker<FSkeinCopyWorker>));
+//	SkeinSourceControlProvider.RegisterWorker("Resolve", FGetSkeinSourceControlWorker::CreateStatic(&CreateWorker<FSkeinResolveWorker>));
+
+	// Bind our source control provider to the editor
+	IModularFeatures::Get().RegisterModularFeature("SourceControl", &SkeinSourceControlProvider);
+}
+
+void FSkeinSourceControlModule::ShutdownModule()
+{
+	// Shut down the provider, as this module is going away
+	SkeinSourceControlProvider.Close();
+
+	// Unbind provider from editor
+	IModularFeatures::Get().UnregisterModularFeature("SourceControl", &SkeinSourceControlProvider);
+}
+
+IMPLEMENT_MODULE(FSkeinSourceControlModule, SkeinSourceControl);
+
+#undef LOCTEXT_NAMESPACE
diff --git a/Engine/Plugins/Developer/SkeinSourceControl/Source/SkeinSourceControl/Private/SkeinSourceControlModule.h b/Engine/Plugins/Developer/SkeinSourceControl/Source/SkeinSourceControl/Private/SkeinSourceControlModule.h
new file mode 100644
--- /dev/null
+++ b/Engine/Plugins/Developer/SkeinSourceControl/Source/SkeinSourceControl/Private/SkeinSourceControlModule.h
@@ -0,0 +1,25 @@
+// Copyright Epic Games, Inc. All Rights Reserved.
+
+#pragma once
+
+#include "CoreMinimal.h"
+#include "Modules/ModuleInterface.h"
+#include "SkeinSourceControlProvider.h"
+
+class FSkeinSourceControlModule : public IModuleInterface
+{
+public:
+	/** IModuleInterface implementation */
+	virtual void StartupModule() override;
+	virtual void ShutdownModule() override;
+
+	/** Access the Skein source control provider */
+	FSkeinSourceControlProvider& GetProvider()
+	{
+		return SkeinSourceControlProvider;
+	}
+
+private:
+	/** The Skein source control provider */
+	FSkeinSourceControlProvider SkeinSourceControlProvider;
+};
diff --git a/Engine/Plugins/Developer/SkeinSourceControl/Source/SkeinSourceControl/Private/SkeinSourceControlOperations.cpp b/Engine/Plugins/Developer/SkeinSourceControl/Source/SkeinSourceControl/Private/SkeinSourceControlOperations.cpp
new file mode 100644
--- /dev/null
+++ b/Engine/Plugins/Developer/SkeinSourceControl/Source/SkeinSourceControl/Private/SkeinSourceControlOperations.cpp
@@ -0,0 +1,195 @@
+// Copyright Epic Games, Inc. All Rights Reserved.
+
+#include "SkeinSourceControlOperations.h"
+#include "SkeinSourceControlThumbnail.h"
+#include "SkeinSourceControlProvider.h"
+#include "SkeinSourceControlCommand.h"
+#include "SkeinSourceControlModule.h"
+#include "SkeinSourceControlUtils.h"
+#include "SkeinSourceControlState.h"
+#include "SourceControlOperations.h"
+
+#define LOCTEXT_NAMESPACE "SkeinSourceControl"
+
+
+
+static bool UpdateCachedStates(const TArray<FSkeinSourceControlState>& InStates)
+{
+	FSkeinSourceControlModule& SkeinSourceControl = FModuleManager::LoadModuleChecked<FSkeinSourceControlModule>("SkeinSourceControl");
+	FSkeinSourceControlProvider& Provider = SkeinSourceControl.GetProvider();
+
+	int NbStatesUpdated = 0;
+
+	for (const auto& InState : InStates)
+	{
+		TSharedRef<FSkeinSourceControlState, ESPMode::ThreadSafe> State = Provider.GetStateInternal(InState.Filename);
+		if (State->State != InState.State)
+		{
+			State->State = InState.State;
+			State->TimeStamp = FDateTime::Now();
+			NbStatesUpdated++;
+		}
+	}
+
+	return (NbStatesUpdated > 0);
+}
+
+
+
+FName FSkeinConnectWorker::GetName() const
+{
+	return "Connect";
+}
+
+bool FSkeinConnectWorker::Execute(FSkeinSourceControlCommand& InCommand)
+{
+	check(InCommand.Operation->GetName() == GetName());
+
+	if (InCommand.SkeinBinaryPath.IsEmpty())
+	{
+		StaticCastSharedRef<FConnect>(InCommand.Operation)->SetErrorText(LOCTEXT("SkeinNotAvailable", "The Skein Command Line application could not be found."));
+		return false;
+	}
+
+	if (InCommand.SkeinProjectRoot.IsEmpty())
+	{
+		StaticCastSharedRef<FConnect>(InCommand.Operation)->SetErrorText(LOCTEXT("SkeinNotEnabled", "There is no Skein project initialized for this location."));
+		return false;
+	}
+
+	InCommand.bCommandSuccessful = SkeinSourceControlUtils::RunCommand(TEXT("auth login"), InCommand.SkeinBinaryPath, InCommand.SkeinProjectRoot, TArray<FString>(), InCommand.Files, InCommand.InfoMessages, InCommand.ErrorMessages);
+
+	return InCommand.bCommandSuccessful;
+}
+
+bool FSkeinConnectWorker::UpdateStates() const
+{
+	check(IsInGameThread());
+
+	return UpdateCachedStates(States);
+}
+
+
+
+FName FSkeinCheckInWorker::GetName() const
+{
+	return "CheckIn";
+}
+
+bool FSkeinCheckInWorker::Execute(FSkeinSourceControlCommand& InCommand)
+{
+	check(InCommand.Operation->GetName() == GetName());
+
+	InCommand.bCommandSuccessful = SkeinSourceControlUtils::RunCommand(TEXT("projects commit"), InCommand.SkeinBinaryPath, InCommand.SkeinProjectRoot, TArray<FString>(), InCommand.Files, InCommand.InfoMessages, InCommand.ErrorMessages);
+
+	SkeinSourceControlUtils::RunUpdateStatus(InCommand.SkeinBinaryPath, InCommand.SkeinProjectRoot, InCommand.Files, InCommand.ErrorMessages, States);
+
+	return InCommand.bCommandSuccessful;
+}
+
+bool FSkeinCheckInWorker::UpdateStates() const
+{
+	check(IsInGameThread());
+
+	return UpdateCachedStates(States);
+}
+
+
+
+FName FSkeinMarkForAddWorker::GetName() const
+{
+	return "MarkForAdd";
+}
+
+bool FSkeinMarkForAddWorker::Execute(FSkeinSourceControlCommand& InCommand)
+{
+	check(InCommand.Operation->GetName() == GetName());
+
+	InCommand.bCommandSuccessful = SkeinSourceControlUtils::RunCommand(TEXT("assets add"), InCommand.SkeinBinaryPath, InCommand.SkeinProjectRoot, TArray<FString>(), InCommand.Files, InCommand.InfoMessages, InCommand.ErrorMessages);
+
+	SkeinSourceControlUtils::RunUpdateStatus(InCommand.SkeinBinaryPath, InCommand.SkeinProjectRoot, InCommand.Files, InCommand.ErrorMessages, States);
+
+	return InCommand.bCommandSuccessful;
+}
+
+bool FSkeinMarkForAddWorker::UpdateStates() const
+{
+	check(IsInGameThread());
+
+	return UpdateCachedStates(States);
+}
+
+
+
+FName FSkeinDeleteWorker::GetName() const
+{
+	return "Delete";
+}
+
+bool FSkeinDeleteWorker::Execute(FSkeinSourceControlCommand& InCommand)
+{
+	check(InCommand.Operation->GetName() == GetName());
+
+	InCommand.bCommandSuccessful = SkeinSourceControlUtils::RunCommand(TEXT("assets rm"), InCommand.SkeinBinaryPath, InCommand.SkeinProjectRoot, TArray<FString>(), InCommand.Files, InCommand.InfoMessages, InCommand.ErrorMessages);
+
+	SkeinSourceControlUtils::RunUpdateStatus(InCommand.SkeinBinaryPath, InCommand.SkeinProjectRoot, InCommand.Files, InCommand.ErrorMessages, States);
+
+	return InCommand.bCommandSuccessful;
+}
+
+bool FSkeinDeleteWorker::UpdateStates() const
+{
+	check(IsInGameThread());
+
+	return UpdateCachedStates(States);
+}
+
+
+
+FName FSkeinSyncWorker::GetName() const
+{
+	return "Sync";
+}
+
+bool FSkeinSyncWorker::Execute(FSkeinSourceControlCommand& InCommand)
+{
+	check(InCommand.Operation->GetName() == GetName());
+
+	InCommand.bCommandSuccessful = SkeinSourceControlUtils::RunCommand(TEXT("projects pull"), InCommand.SkeinBinaryPath, InCommand.SkeinProjectRoot, TArray<FString>(), TArray<FString>(), InCommand.InfoMessages, InCommand.ErrorMessages);
+
+	SkeinSourceControlUtils::RunUpdateStatus(InCommand.SkeinBinaryPath, InCommand.SkeinProjectRoot, InCommand.Files, InCommand.ErrorMessages, States);
+
+	return InCommand.bCommandSuccessful;
+}
+
+bool FSkeinSyncWorker::UpdateStates() const
+{
+	check(IsInGameThread());
+
+	return UpdateCachedStates(States);
+}
+
+
+
+FName FSkeinUpdateStatusWorker::GetName() const
+{
+	return "UpdateStatus";
+}
+
+bool FSkeinUpdateStatusWorker::Execute(FSkeinSourceControlCommand& InCommand)
+{
+	check(InCommand.Operation->GetName() == GetName());
+
+	InCommand.bCommandSuccessful = SkeinSourceControlUtils::RunUpdateStatus(InCommand.SkeinBinaryPath, InCommand.SkeinProjectRoot, InCommand.Files, InCommand.ErrorMessages, States);
+
+	return InCommand.bCommandSuccessful;
+}
+
+bool FSkeinUpdateStatusWorker::UpdateStates() const
+{
+	check(IsInGameThread());
+
+	return UpdateCachedStates(States);
+}
+
+#undef LOCTEXT_NAMESPACE
\ No newline at end of file
diff --git a/Engine/Plugins/Developer/SkeinSourceControl/Source/SkeinSourceControl/Private/SkeinSourceControlOperations.h b/Engine/Plugins/Developer/SkeinSourceControl/Source/SkeinSourceControl/Private/SkeinSourceControlOperations.h
new file mode 100644
--- /dev/null
+++ b/Engine/Plugins/Developer/SkeinSourceControl/Source/SkeinSourceControl/Private/SkeinSourceControlOperations.h
@@ -0,0 +1,95 @@
+#pragma once
+
+#include "CoreMinimal.h"
+#include "ISkeinSourceControlWorker.h"
+#include "SkeinSourceControlState.h"
+
+/** Connect to the source control. */
+class FSkeinConnectWorker : public ISkeinSourceControlWorker
+{
+public:
+	virtual ~FSkeinConnectWorker() {}
+	/** ISkeinSourceControlWorker interface */
+	virtual FName GetName() const override;
+	virtual bool Execute(class FSkeinSourceControlCommand& InCommand) override;
+	virtual bool UpdateStates() const override;
+
+public:
+	/** Temporary states for results */
+	TArray<FSkeinSourceControlState> States;
+};
+
+/** Commit (check-in) a set of files. */
+class FSkeinCheckInWorker : public ISkeinSourceControlWorker
+{
+public:
+	virtual ~FSkeinCheckInWorker() {}
+	/** ISkeinSourceControlWorker interface */
+	virtual FName GetName() const override;
+	virtual bool Execute(class FSkeinSourceControlCommand& InCommand) override;
+	virtual bool UpdateStates() const override;
+
+public:
+	/** Temporary states for results */
+	TArray<FSkeinSourceControlState> States;
+};
+
+/** Add an untracked file to source control. */
+class FSkeinMarkForAddWorker : public ISkeinSourceControlWorker
+{
+public:
+	virtual ~FSkeinMarkForAddWorker() {}
+	/** ISkeinSourceControlWorker interface */
+	virtual FName GetName() const override;
+	virtual bool Execute(class FSkeinSourceControlCommand& InCommand) override;
+	virtual bool UpdateStates() const override;
+
+public:
+	/** Temporary states for results */
+	TArray<FSkeinSourceControlState> States;
+};
+
+/** Delete a file and remove it from source control. */
+class FSkeinDeleteWorker : public ISkeinSourceControlWorker
+{
+public:
+	virtual ~FSkeinDeleteWorker() {}
+	/** ISkeinSourceControlWorker interface */
+	virtual FName GetName() const override;
+	virtual bool Execute(class FSkeinSourceControlCommand& InCommand) override;
+	virtual bool UpdateStates() const override;
+
+public:
+	/** Temporary states for results */
+	TArray<FSkeinSourceControlState> States;
+};
+
+/** Skein synchronize to the active snapshot set in WebUI. */
+class FSkeinSyncWorker : public ISkeinSourceControlWorker
+{
+public:
+	virtual ~FSkeinSyncWorker() {}
+	/** ISkeinSourceControlWorker interface */
+	virtual FName GetName() const override;
+	virtual bool Execute(class FSkeinSourceControlCommand& InCommand) override;
+	virtual bool UpdateStates() const override;
+
+public:
+	/** Temporary states for results */
+	TArray<FSkeinSourceControlState> States;
+};
+
+/** Get source control status of files. */
+class FSkeinUpdateStatusWorker : public ISkeinSourceControlWorker
+{
+public:
+	virtual ~FSkeinUpdateStatusWorker() {}
+	/** ISkeinSourceControlWorker interface */
+	virtual FName GetName() const override;
+	virtual bool Execute(class FSkeinSourceControlCommand& InCommand) override;
+	virtual bool UpdateStates() const override;
+
+public:
+	/** Temporary states for results */
+	TArray<FSkeinSourceControlState> States;
+};
\ No newline at end of file
diff --git a/Engine/Plugins/Developer/SkeinSourceControl/Source/SkeinSourceControl/Private/SkeinSourceControlProvider.cpp b/Engine/Plugins/Developer/SkeinSourceControl/Source/SkeinSourceControl/Private/SkeinSourceControlProvider.cpp
new file mode 100644
--- /dev/null
+++ b/Engine/Plugins/Developer/SkeinSourceControl/Source/SkeinSourceControl/Private/SkeinSourceControlProvider.cpp
@@ -0,0 +1,348 @@
+// Copyright Epic Games, Inc. All Rights Reserved.
+
+#include "SkeinSourceControlProvider.h"
+#include "SkeinSourceControlThumbnail.h"
+#include "SkeinSourceControlCommand.h"
+#include "SkeinSourceControlState.h"
+#include "SkeinSourceControlUtils.h"
+#include "ScopedSourceControlProgress.h"
+#include "SourceControlOperations.h"
+#include "SourceControlHelpers.h"
+#include "Logging/MessageLog.h"
+
+#if SOURCE_CONTROL_WITH_SLATE
+#include "SSkeinSourceControlSettings.h"
+#endif
+
+#define LOCTEXT_NAMESPACE "SkeinSourceControlProvider"
+
+void FSkeinSourceControlProvider::Init(bool bForceConnection)
+{
+	bSkeinBinaryFound = SkeinSourceControlUtils::IsSkeinAvailable();
+	bSkeinProjectFound = SkeinSourceControlUtils::IsSkeinProjectFound(FPaths::ProjectDir(), ProjectRoot, ProjectName);
+	BinaryPath = SkeinSourceControlUtils::FindSkeinBinaryPath();
+}
+
+void FSkeinSourceControlProvider::Close()
+{
+	BinaryPath.Empty();
+	ProjectName.Empty();
+	ProjectRoot.Empty();
+	bSkeinProjectFound = false;
+	bSkeinBinaryFound = false;
+}
+
+FText FSkeinSourceControlProvider::GetStatusText() const
+{
+	FFormatNamedArguments Args;
+	Args.Add(TEXT("IsAvailable"), IsAvailable() ? LOCTEXT("Yes", "Yes") : LOCTEXT("No", "No"));
+	Args.Add(TEXT("IsEnabled"), IsEnabled() ? LOCTEXT("Yes", "Yes") : LOCTEXT("No", "No"));
+	Args.Add(TEXT("ProjectName"), FText::FromString(ProjectName));
+	Args.Add(TEXT("ProjectRoot"), FText::FromString(ProjectRoot));
+
+	return FText::Format(LOCTEXT("SkeinStatusText", "Available: {IsAvailable}\nEnabled: {IsEnabled}\nProjectName: {ProjectName}\nProjectRoot: {ProjectRoot}"), Args);
+}
+
+bool FSkeinSourceControlProvider::IsAvailable() const
+{
+	return bSkeinBinaryFound;
+}
+
+bool FSkeinSourceControlProvider::IsEnabled() const
+{
+	return bSkeinBinaryFound && bSkeinProjectFound;
+}
+
+const FName& FSkeinSourceControlProvider::GetName(void) const
+{
+	static FName ProviderName("Skein");
+	return ProviderName;
+}
+
+ECommandResult::Type FSkeinSourceControlProvider::GetState(const TArray<FString>& InFiles, TArray<TSharedRef<ISourceControlState, ESPMode::ThreadSafe>>& OutState, EStateCacheUsage::Type InStateCacheUsage)
+{
+	if (!IsEnabled())
+	{
+		return ECommandResult::Failed;
+	}
+
+	TArray<FString> AbsoluteFiles = SourceControlHelpers::AbsoluteFilenames(InFiles);
+
+	if (InStateCacheUsage == EStateCacheUsage::ForceUpdate)
+	{
+		Execute(ISourceControlOperation::Create<FUpdateStatus>(), AbsoluteFiles);
+	}
+
+	for (TArray<FString>::TConstIterator It(AbsoluteFiles); It; It++)
+	{
+		OutState.Add(GetStateInternal(*It));
+	}
+
+	return ECommandResult::Succeeded;
+}
+
+ECommandResult::Type FSkeinSourceControlProvider::GetState(const TArray<FSourceControlChangelistRef>& InChangelists, TArray<FSourceControlChangelistStateRef>& OutState, EStateCacheUsage::Type InStateCacheUsage)
+{
+	return ECommandResult::Failed;
+}
+
+TArray<FSourceControlStateRef> FSkeinSourceControlProvider::GetCachedStateByPredicate(TFunctionRef<bool(const FSourceControlStateRef&)> Predicate) const
+{
+	TArray<FSourceControlStateRef> Result;
+	for (const auto& CacheItem : StateCache)
+	{
+		FSourceControlStateRef State = CacheItem.Value;
+		if (Predicate(State))
+		{
+			Result.Add(State);
+		}
+	}
+	return Result;
+}
+
+bool FSkeinSourceControlProvider::RemoveFileFromCache(const FString& Filename)
+{
+	return StateCache.Remove(Filename) > 0;
+}
+
+FDelegateHandle FSkeinSourceControlProvider::RegisterSourceControlStateChanged_Handle(const FSourceControlStateChanged::FDelegate& SourceControlStateChanged)
+{
+	return OnSourceControlStateChanged.Add(SourceControlStateChanged);
+}
+
+void FSkeinSourceControlProvider::UnregisterSourceControlStateChanged_Handle(FDelegateHandle Handle)
+{
+	OnSourceControlStateChanged.Remove(Handle);
+}
+
+ECommandResult::Type FSkeinSourceControlProvider::Execute(const FSourceControlOperationRef& InOperation, FSourceControlChangelistPtr InChangelist, const TArray<FString>& InFiles, EConcurrency::Type InConcurrency, const FSourceControlOperationComplete& InOperationCompleteDelegate)
+{
+	// Query to see if the we allow this operation
+	TSharedPtr<ISkeinSourceControlWorker, ESPMode::ThreadSafe> Worker = CreateWorker(InOperation->GetName());
+	if (!Worker.IsValid())
+	{
+		// This operation is unsupported by this source control provider
+		FFormatNamedArguments Arguments;
+		Arguments.Add(TEXT("OperationName"), FText::FromName(InOperation->GetName()));
+		Arguments.Add(TEXT("ProviderName"), FText::FromName(GetName()));
+		FText Message = FText::Format(LOCTEXT("UnsupportedOperation", "Operation '{OperationName}' not supported by source control provider '{ProviderName}'"), Arguments);
+
+		FMessageLog("SourceControl").Error(Message);
+		InOperation->AddErrorMessge(Message);
+
+		InOperationCompleteDelegate.ExecuteIfBound(InOperation, ECommandResult::Failed);
+		return ECommandResult::Failed;
+	}
+
+	TArray<FString> AbsoluteFiles = SourceControlHelpers::AbsoluteFilenames(InFiles);
+
+	// Fire off operation
+	if (InConcurrency == EConcurrency::Synchronous)
+	{
+		FSkeinSourceControlCommand* Command = new FSkeinSourceControlCommand(InOperation, Worker.ToSharedRef());
+		Command->bAutoDelete = false;
+		Command->Files = AbsoluteFiles;
+		Command->OperationCompleteDelegate = InOperationCompleteDelegate;
+		return ExecuteSynchronousCommand(*Command, InOperation->GetInProgressString());
+	}
+	else
+	{
+		FSkeinSourceControlCommand* Command = new FSkeinSourceControlCommand(InOperation, Worker.ToSharedRef());
+		Command->bAutoDelete = true;
+		Command->Files = AbsoluteFiles;
+		Command->OperationCompleteDelegate = InOperationCompleteDelegate;
+		return IssueCommand(*Command);
+	}
+
+	return ECommandResult::Failed;
+}
+
+bool FSkeinSourceControlProvider::CanCancelOperation(const FSourceControlOperationRef& InOperation) const
+{
+	return false;
+}
+
+void FSkeinSourceControlProvider::CancelOperation(const FSourceControlOperationRef& InOperation)
+{
+}
+
+bool FSkeinSourceControlProvider::UsesLocalReadOnlyState() const
+{
+	return false;
+}
+
+bool FSkeinSourceControlProvider::UsesChangelists() const
+{
+	return false;
+}
+
+bool FSkeinSourceControlProvider::UsesCheckout() const
+{
+	return false;
+}
+
+void FSkeinSourceControlProvider::Tick()
+{
+	check(IsInGameThread());
+
+	bool bStatesUpdated = false;
+	for (int32 CommandIndex = 0; CommandIndex < CommandQueue.Num(); ++CommandIndex)
+	{
+		FSkeinSourceControlCommand& Command = *CommandQueue[CommandIndex];
+		if (Command.bExecuteProcessed)
+		{
+			// Remove command from the queue
+			CommandQueue.RemoveAt(CommandIndex);
+
+			// Let command update the states of any files
+			bStatesUpdated |= Command.Worker->UpdateStates();
+
+			// Dump any messages to output log
+			FMessageLog SourceControlLog("SourceControl");
+			for (int32 ErrorIndex = 0; ErrorIndex < Command.ErrorMessages.Num(); ++ErrorIndex)
+			{
+				SourceControlLog.Error(FText::FromString(Command.ErrorMessages[ErrorIndex]));
+			}
+			for (int32 InfoIndex = 0; InfoIndex < Command.InfoMessages.Num(); ++InfoIndex)
+			{
+				SourceControlLog.Info(FText::FromString(Command.InfoMessages[InfoIndex]));
+			}
+
+			Command.ReturnResults();
+
+			// Commands that are left in the array during a tick need to be deleted
+			if (Command.bAutoDelete)
+			{
+				// Only delete commands that are not running 'synchronously'
+				delete &Command;
+			}
+
+			// Only do one command per tick loop, as we dont want concurrent modification
+			// of the command queue (which can happen in the completion delegate)
+			break;
+		}
+	}
+
+	if (bStatesUpdated)
+	{
+		OnSourceControlStateChanged.Broadcast();
+	}
+}
+
+TArray< TSharedRef<ISourceControlLabel> > FSkeinSourceControlProvider::GetLabels(const FString& InMatchingSpec) const
+{
+	return TArray< TSharedRef<ISourceControlLabel> >();
+}
+
+TArray<FSourceControlChangelistRef> FSkeinSourceControlProvider::GetChangelists(EStateCacheUsage::Type InStateCacheUsage)
+{
+	return TArray<FSourceControlChangelistRef>();
+}
+
+#if SOURCE_CONTROL_WITH_SLATE
+TSharedRef<class SWidget> FSkeinSourceControlProvider::MakeSettingsWidget() const
+{
+	return SNew(SSkeinSourceControlSettings);
+}
+#endif // SOURCE_CONTROL_WITH_SLATE
+
+void FSkeinSourceControlProvider::RegisterWorker(const FName& InName, const FGetSkeinSourceControlWorker& InDelegate)
+{
+	WorkersMap.Add(InName, InDelegate);
+}
+
+TSharedPtr<ISkeinSourceControlWorker, ESPMode::ThreadSafe> FSkeinSourceControlProvider::CreateWorker(const FName& InOperationName) const
+{
+	const FGetSkeinSourceControlWorker* Operation = WorkersMap.Find(InOperationName);
+	if (Operation != nullptr)
+	{
+		return Operation->Execute(); // Creates the ISkeinSourceControlWorker.
+	}
+
+	return nullptr;
+}
+
+ECommandResult::Type FSkeinSourceControlProvider::ExecuteSynchronousCommand(FSkeinSourceControlCommand& InCommand, const FText& Task)
+{
+	ECommandResult::Type Result = ECommandResult::Failed;
+
+	// Display the progress dialog if a string was provided
+	{
+		FScopedSourceControlProgress Progress(Task);
+
+		// Issue the command asynchronously...
+		IssueCommand(InCommand);
+
+		// ... then wait for its completion (thus making it synchronous)
+		while (!InCommand.bExecuteProcessed)
+		{
+			// Tick the command queue and update progress.
+			Tick();
+
+			Progress.Tick();
+
+			// Sleep for a bit so we don't busy-wait so much.
+			FPlatformProcess::Sleep(0.01f);
+		}
+
+		// Always do one more Tick() to make sure the command queue is cleaned up.
+		Tick();
+
+		if (InCommand.bCommandSuccessful)
+		{
+			Result = ECommandResult::Succeeded;
+		}
+	}
+
+	// Delete the command now (asynchronous commands are deleted in the Tick() method)
+	check(!InCommand.bAutoDelete);
+
+	// Ensure commands that are not auto deleted do not end up in the command queue
+	if (CommandQueue.Contains(&InCommand))
+	{
+		CommandQueue.Remove(&InCommand);
+	}
+	delete &InCommand;
+
+	return Result;
+}
+
+ECommandResult::Type FSkeinSourceControlProvider::IssueCommand(FSkeinSourceControlCommand& InCommand)
+{
+	if (GThreadPool != nullptr)
+	{
+		// Queue this to our worker thread(s) for resolving
+		GThreadPool->AddQueuedWork(&InCommand);
+		CommandQueue.Add(&InCommand);
+		return ECommandResult::Succeeded;
+	}
+	else
+	{
+		FText Message(LOCTEXT("NoSCCThreads", "There are no threads available to process the source control command."));
+
+		FMessageLog("SourceControl").Error(Message);
+		InCommand.bCommandSuccessful = false;
+		InCommand.Operation->AddErrorMessge(Message);
+
+		return InCommand.ReturnResults();
+	}
+}
+
+TSharedRef<FSkeinSourceControlState, ESPMode::ThreadSafe> FSkeinSourceControlProvider::GetStateInternal(const FString& Filename)
+{
+	TSharedRef<FSkeinSourceControlState, ESPMode::ThreadSafe>* State = StateCache.Find(Filename);
+	if (State != NULL)
+	{
+		// Found cached item
+		return (*State);
+	}
+	else
+	{
+		// Cache an unknown state for this item
+		TSharedRef<FSkeinSourceControlState, ESPMode::ThreadSafe> NewState = MakeShareable(new FSkeinSourceControlState(Filename));
+		StateCache.Add(Filename, NewState);
+		return NewState;
+	}
+}
+
+#undef LOCTEXT_NAMESPACE
diff --git a/Engine/Plugins/Developer/SkeinSourceControl/Source/SkeinSourceControl/Private/SkeinSourceControlProvider.h b/Engine/Plugins/Developer/SkeinSourceControl/Source/SkeinSourceControl/Private/SkeinSourceControlProvider.h
new file mode 100644
--- /dev/null
+++ b/Engine/Plugins/Developer/SkeinSourceControl/Source/SkeinSourceControl/Private/SkeinSourceControlProvider.h
@@ -0,0 +1,122 @@
+// Copyright Epic Games, Inc. All Rights Reserved.
+
+#pragma once
+
+#include "CoreMinimal.h"
+#include "ISourceControlProvider.h"
+#include "ISkeinSourceControlWorker.h"
+
+class FSkeinSourceControlCommand;
+class FSkeinSourceControlState;
+
+DECLARE_DELEGATE_RetVal(FSkeinSourceControlWorkerRef, FGetSkeinSourceControlWorker)
+
+class FSkeinSourceControlProvider : public ISourceControlProvider
+{
+public:
+	/** Constructor */
+	FSkeinSourceControlProvider()
+		: bSkeinBinaryFound(false)
+		, bSkeinProjectFound(false)
+	{
+	}
+
+	/** ISourceControlProvider implementation */
+	virtual void Init(bool bForceConnection = true) override;
+	virtual void Close() override;
+	virtual FText GetStatusText() const override;
+	virtual bool IsAvailable() const override;
+	virtual bool IsEnabled() const override;
+	virtual const FName& GetName(void) const override;
+	virtual bool QueryStateBranchConfig(const FString& ConfigSrc, const FString& ConfigDest) override { return false; }
+	virtual void RegisterStateBranches(const TArray<FString>& BranchNames, const FString& ContentRoot) override {}
+	virtual int32 GetStateBranchIndex(const FString& BranchName) const override { return INDEX_NONE; }
+	virtual ECommandResult::Type GetState(const TArray<FString>& InFiles, TArray<FSourceControlStateRef>& OutState, EStateCacheUsage::Type InStateCacheUsage) override;
+	virtual ECommandResult::Type GetState(const TArray<FSourceControlChangelistRef>& InChangelists, TArray<FSourceControlChangelistStateRef>& OutState, EStateCacheUsage::Type InStateCacheUsage) override;
+	virtual TArray<FSourceControlStateRef> GetCachedStateByPredicate(TFunctionRef<bool(const FSourceControlStateRef&)> Predicate) const override;
+	virtual FDelegateHandle RegisterSourceControlStateChanged_Handle(const FSourceControlStateChanged::FDelegate& SourceControlStateChanged) override;
+	virtual void UnregisterSourceControlStateChanged_Handle(FDelegateHandle Handle) override;
+	virtual ECommandResult::Type Execute(const FSourceControlOperationRef& InOperation, FSourceControlChangelistPtr InChangelist, const TArray<FString>& InFiles, EConcurrency::Type InConcurrency = EConcurrency::Synchronous, const FSourceControlOperationComplete& InOperationCompleteDelegate = FSourceControlOperationComplete()) override;
+	virtual bool CanCancelOperation(const FSourceControlOperationRef& InOperation) const override;
+	virtual void CancelOperation(const FSourceControlOperationRef& InOperation) override;
+	virtual bool UsesLocalReadOnlyState() const override;
+	virtual bool UsesChangelists() const override;
+	virtual bool UsesCheckout() const override;
+	virtual void Tick() override;
+	virtual TArray< TSharedRef<class ISourceControlLabel> > GetLabels(const FString& InMatchingSpec) const override;
+	virtual TArray<FSourceControlChangelistRef> GetChangelists(EStateCacheUsage::Type InStateCacheUsage) override;
+#if SOURCE_CONTROL_WITH_SLATE
+	virtual TSharedRef<class SWidget> MakeSettingsWidget() const override;
+#endif // SOURCE_CONTROL_WITH_SLATE
+
+	/** Allow calls to all ISourceControlProvider::Execute variants */
+	using ISourceControlProvider::Execute;
+
+	/** Get the path to the Skein CLI binary */
+	inline const FString& GetSkeinBinaryPath() const
+	{
+		return BinaryPath;
+	}
+
+	/** Get the path to the root of the Skein project: can be the ProjectDir itself, or any parent directory */
+	inline const FString& GetSkeinProjectRoot() const
+	{
+		return ProjectRoot;
+	}
+
+	/** Remove a named file from the state cache */
+	bool RemoveFileFromCache(const FString& Filename);
+
+	/** Helper function used to update state cache */
+	TSharedRef<FSkeinSourceControlState, ESPMode::ThreadSafe> GetStateInternal(const FString& Filename);
+
+	/**
+	 * Register a worker with the provider.
+	 * This is used internally so the provider can maintain a map of all available operations.
+	 */
+	void RegisterWorker(const FName& InName, const FGetSkeinSourceControlWorker& InDelegate);
+
+private:
+
+	/** Helper function for Execute() */
+	TSharedPtr<class ISkeinSourceControlWorker, ESPMode::ThreadSafe> CreateWorker(const FName& InOperationName) const;
+
+	/**
+	 * Helper function for running command 'synchronously'.
+	 * This really doesn't execute synchronously; rather it adds the command to the queue & does not return until
+	 * the command is completed.
+	 */
+	ECommandResult::Type ExecuteSynchronousCommand(FSkeinSourceControlCommand& InCommand, const FText& Task);
+
+	/** Issue a command asynchronously if possible. */
+	ECommandResult::Type IssueCommand(FSkeinSourceControlCommand& InCommand);
+
+private:
+
+	/** Is Skein service available */
+	bool bSkeinBinaryFound;
+
+	/** Is Skein project found */
+	bool bSkeinProjectFound;
+
+	/** Skein binary path */
+	FString BinaryPath;
+
+	/** Skein project root */
+	FString ProjectRoot;
+
+	/** Skein project name */
+	FString ProjectName;
+
+	/** State cache */
+	TMap<FString, TSharedRef<FSkeinSourceControlState, ESPMode::ThreadSafe>> StateCache;
+
+	/** The currently registered source control operations */
+	TMap<FName, FGetSkeinSourceControlWorker> WorkersMap;
+
+	/** Queue for commands given by the main thread */
+	TArray<FSkeinSourceControlCommand*> CommandQueue;
+
+	/** For notifying when the source control states in the cache have changed */
+	FSourceControlStateChanged OnSourceControlStateChanged;
+};
\ No newline at end of file
diff --git a/Engine/Plugins/Developer/SkeinSourceControl/Source/SkeinSourceControl/Private/SkeinSourceControlState.cpp b/Engine/Plugins/Developer/SkeinSourceControl/Source/SkeinSourceControl/Private/SkeinSourceControlState.cpp
new file mode 100644
--- /dev/null
+++ b/Engine/Plugins/Developer/SkeinSourceControl/Source/SkeinSourceControl/Private/SkeinSourceControlState.cpp
@@ -0,0 +1,237 @@
+// Copyright Epic Games, Inc. All Rights Reserved.
+
+#include "SkeinSourceControlState.h"
+#include "Styling/AppStyle.h"
+
+#define LOCTEXT_NAMESPACE "SkeinSourceControl.State"
+
+int32 FSkeinSourceControlState::GetHistorySize() const
+{
+	return 0;
+}
+
+TSharedPtr<class ISourceControlRevision, ESPMode::ThreadSafe> FSkeinSourceControlState::GetHistoryItem(int32 HistoryIndex) const
+{
+	return nullptr;
+}
+
+TSharedPtr<class ISourceControlRevision, ESPMode::ThreadSafe> FSkeinSourceControlState::FindHistoryRevision(int32 RevisionNumber) const
+{
+	return nullptr;
+}
+
+TSharedPtr<class ISourceControlRevision, ESPMode::ThreadSafe> FSkeinSourceControlState::FindHistoryRevision(const FString& InRevision) const
+{
+	return nullptr;
+}
+
+TSharedPtr<class ISourceControlRevision, ESPMode::ThreadSafe> FSkeinSourceControlState::GetBaseRevForMerge() const
+{
+	return nullptr;
+}
+
+FSlateIcon FSkeinSourceControlState::GetIcon() const
+{
+	switch (State)
+	{
+	case ESkeinState::Modified:
+		return FSlateIcon(FAppStyle::GetAppStyleSetName(), "Subversion.CheckedOut");
+	case ESkeinState::Added:
+		return FSlateIcon(FAppStyle::GetAppStyleSetName(), "Subversion.OpenForAdd");
+	case ESkeinState::Renamed:
+	case ESkeinState::Copied:
+		return FSlateIcon(FAppStyle::GetAppStyleSetName(), "Subversion.Branched");
+	case ESkeinState::Deleted: // Deleted & Missing files do not show in Content Browser
+	case ESkeinState::Missing:
+		return FSlateIcon(FAppStyle::GetAppStyleSetName(), "Subversion.MarkedForDelete");
+	case ESkeinState::Conflicted:
+		return FSlateIcon(FAppStyle::GetAppStyleSetName(), "Subversion.NotAtHeadRevision");
+	case ESkeinState::NotControlled:
+		return FSlateIcon(FAppStyle::GetAppStyleSetName(), "Subversion.NotInDepot");
+	case ESkeinState::Unknown:
+	case ESkeinState::Unchanged: // Unchanged is the same as "Pristine" (not checked out) for Perforce, ie no icon
+	case ESkeinState::Ignored:
+	default:
+		return FSlateIcon();
+	}
+
+	return FSlateIcon();
+}
+
+
+FText FSkeinSourceControlState::GetDisplayName() const
+{
+	switch (State)
+	{
+	case ESkeinState::Unknown:
+		return LOCTEXT("Unknown", "Unknown");
+	case ESkeinState::Unchanged:
+		return LOCTEXT("Unchanged", "Unchanged");
+	case ESkeinState::Added:
+		return LOCTEXT("Added", "Added");
+	case ESkeinState::Deleted:
+		return LOCTEXT("Deleted", "Deleted");
+	case ESkeinState::Modified:
+		return LOCTEXT("Modified", "Modified");
+	case ESkeinState::Renamed:
+		return LOCTEXT("Renamed", "Renamed");
+	case ESkeinState::Copied:
+		return LOCTEXT("Copied", "Copied");
+	case ESkeinState::Conflicted:
+		return LOCTEXT("ContentsConflict", "Contents Conflict");
+	case ESkeinState::Ignored:
+		return LOCTEXT("Ignored", "Ignored");
+	case ESkeinState::NotControlled:
+		return LOCTEXT("NotControlled", "Not Under Source Control");
+	case ESkeinState::Missing:
+		return LOCTEXT("Missing", "Missing");
+	}
+
+	return FText();
+}
+
+FText FSkeinSourceControlState::GetDisplayTooltip() const
+{
+	switch (State)
+	{
+	case ESkeinState::Unknown:
+		return LOCTEXT("Unknown_Tooltip", "Unknown source control state");
+	case ESkeinState::Unchanged:
+		return LOCTEXT("Pristine_Tooltip", "There are no modifications");
+	case ESkeinState::Added:
+		return LOCTEXT("Added_Tooltip", "Item is scheduled for addition");
+	case ESkeinState::Deleted:
+		return LOCTEXT("Deleted_Tooltip", "Item is scheduled for deletion");
+	case ESkeinState::Modified:
+		return LOCTEXT("Modified_Tooltip", "Item has been modified");
+	case ESkeinState::Renamed:
+		return LOCTEXT("Renamed_Tooltip", "Item has been renamed");
+	case ESkeinState::Copied:
+		return LOCTEXT("Copied_Tooltip", "Item has been copied");
+	case ESkeinState::Conflicted:
+		return LOCTEXT("ContentsConflict_Tooltip", "The contents of the item conflict with updates received from the repository.");
+	case ESkeinState::Ignored:
+		return LOCTEXT("Ignored_Tooltip", "Item is being ignored.");
+	case ESkeinState::NotControlled:
+		return LOCTEXT("NotControlled_Tooltip", "Item is not under version control.");
+	case ESkeinState::Missing:
+		return LOCTEXT("Missing_Tooltip", "Item is missing (e.g., you moved or deleted it without using Skein). This also indicates that a directory is incomplete (a checkout or update was interrupted).");
+	}
+
+	return FText();
+}
+
+const FString& FSkeinSourceControlState::GetFilename() const
+{
+	return Filename;
+}
+
+const FDateTime& FSkeinSourceControlState::GetTimeStamp() const
+{
+	return TimeStamp;
+}
+
+// Deleted and Missing assets cannot appear in the Content Browser, but the do in the Submit files to Source Control window!
+bool FSkeinSourceControlState::CanCheckIn() const
+{
+	return State == ESkeinState::Added
+		|| State == ESkeinState::Deleted
+		|| State == ESkeinState::Missing
+		|| State == ESkeinState::Modified
+		|| State == ESkeinState::Renamed;
+}
+
+bool FSkeinSourceControlState::CanCheckout() const
+{
+	return false; // With Skein all tracked files in the working copy are always already checked-out (as opposed to Perforce)
+}
+
+bool FSkeinSourceControlState::IsCheckedOut() const
+{
+	return IsSourceControlled(); // With Skein all tracked files in the working copy are always checked-out (as opposed to Perforce)
+}
+
+bool FSkeinSourceControlState::IsCheckedOutOther(FString* Who) const
+{
+	return false; // Skein does not lock checked-out files as Perforce does
+}
+
+bool FSkeinSourceControlState::IsCurrent() const
+{
+	return true; // Could check the state of the HEAD versus the state of tracked branch on remote
+}
+
+bool FSkeinSourceControlState::IsSourceControlled() const
+{
+	return State != ESkeinState::NotControlled && State != ESkeinState::Ignored && State != ESkeinState::Unknown;
+}
+
+bool FSkeinSourceControlState::IsAdded() const
+{
+	return State == ESkeinState::Added;
+}
+
+bool FSkeinSourceControlState::IsDeleted() const
+{
+	return State == ESkeinState::Deleted || State == ESkeinState::Missing;
+}
+
+bool FSkeinSourceControlState::IsIgnored() const
+{
+	return State == ESkeinState::Ignored;
+}
+
+bool FSkeinSourceControlState::CanEdit() const
+{
+	return true; // With Skein all files in the working copy are always editable (as opposed to Perforce)
+}
+
+bool FSkeinSourceControlState::CanDelete() const
+{
+	return IsSourceControlled() && IsCurrent();
+}
+
+bool FSkeinSourceControlState::IsUnknown() const
+{
+	return State == ESkeinState::Unknown;
+}
+
+bool FSkeinSourceControlState::IsModified() const
+{
+	// Warning: for Perforce, a checked-out file is locked for modification (whereas with Skein all tracked files are checked-out),
+	// so for a clean "check-in" (commit) checked-out files unmodified should be removed from the changeset (the index)
+	// http://stackoverflow.com/questions/12357971/what-does-revert-unchanged-files-mean-in-perforce
+	//
+	// Thus, before check-in UE Editor call RevertUnchangedFiles() in PromptForCheckin() and CheckinFiles().
+	//
+	// So here we must take care to enumerate all states that need to be commited,
+	// all other will be discarded :
+	//  - Unknown
+	//  - Unchanged
+	//  - NotControlled
+	//  - Ignored
+	return State == ESkeinState::Added
+		|| State == ESkeinState::Deleted
+		|| State == ESkeinState::Modified
+		|| State == ESkeinState::Renamed
+		|| State == ESkeinState::Copied
+		|| State == ESkeinState::Conflicted
+		|| State == ESkeinState::Missing;
+}
+
+bool FSkeinSourceControlState::CanAdd() const
+{
+	return State == ESkeinState::NotControlled;
+}
+
+bool FSkeinSourceControlState::IsConflicted() const
+{
+	return State == ESkeinState::Conflicted;
+}
+
+bool FSkeinSourceControlState::CanRevert() const
+{
+	return CanCheckIn();
+}
+
+#undef LOCTEXT_NAMESPACE
diff --git a/Engine/Plugins/Developer/SkeinSourceControl/Source/SkeinSourceControl/Private/SkeinSourceControlState.h b/Engine/Plugins/Developer/SkeinSourceControl/Source/SkeinSourceControl/Private/SkeinSourceControlState.h
new file mode 100644
--- /dev/null
+++ b/Engine/Plugins/Developer/SkeinSourceControl/Source/SkeinSourceControl/Private/SkeinSourceControlState.h
@@ -0,0 +1,79 @@
+// Copyright Epic Games, Inc. All Rights Reserved.
+
+#pragma once
+
+namespace ESkeinState
+{
+	// SKEIN_TODO: Re-evaluate each of the possible states used here. Some of them might not be applicable to Skein.
+	//			   The Skein POC Tech Spec does not specify what the possible states for a file can be.
+
+	enum Type
+	{
+		Unknown,
+		Unchanged, // called "clean" in SVN, "Pristine" in Perforce
+		Added,
+		Deleted,
+		Modified,
+		Renamed,
+		Copied,
+		Missing,
+		Conflicted,
+		NotControlled,
+		Ignored,
+	};
+}
+
+class FSkeinSourceControlState : public ISourceControlState, public TSharedFromThis<FSkeinSourceControlState, ESPMode::ThreadSafe>
+{
+public:
+	FSkeinSourceControlState(const FString& InLocalFilename)
+		: Filename(InLocalFilename)
+		, TimeStamp(0)
+		, State(ESkeinState::Unknown)
+	{
+	}
+
+	/** ISourceControlState interface */
+	virtual int32 GetHistorySize() const override;
+	virtual TSharedPtr<class ISourceControlRevision, ESPMode::ThreadSafe> GetHistoryItem(int32 HistoryIndex) const override;
+	virtual TSharedPtr<class ISourceControlRevision, ESPMode::ThreadSafe> FindHistoryRevision(int32 RevisionNumber) const override;
+	virtual TSharedPtr<class ISourceControlRevision, ESPMode::ThreadSafe> FindHistoryRevision(const FString& InRevision) const override;
+	virtual TSharedPtr<class ISourceControlRevision, ESPMode::ThreadSafe> GetBaseRevForMerge() const override;
+	virtual FSlateIcon GetIcon() const override;
+	virtual FText GetDisplayName() const override;
+	virtual FText GetDisplayTooltip() const override;
+	virtual const FString& GetFilename() const override;
+	virtual const FDateTime& GetTimeStamp() const override;
+	virtual bool CanCheckIn() const override;
+	virtual bool CanCheckout() const override;
+	virtual bool IsCheckedOut() const override;
+	virtual bool IsCheckedOutOther(FString* Who = nullptr) const override;
+	virtual bool IsCheckedOutInOtherBranch(const FString& CurrentBranch = FString()) const override { return false; }
+	virtual bool IsModifiedInOtherBranch(const FString& CurrentBranch = FString()) const override { return false; }
+	virtual bool IsCheckedOutOrModifiedInOtherBranch(const FString& CurrentBranch = FString()) const override { return IsCheckedOutInOtherBranch(CurrentBranch) || IsModifiedInOtherBranch(CurrentBranch); }
+	virtual TArray<FString> GetCheckedOutBranches() const override { return TArray<FString>(); }
+	virtual FString GetOtherUserBranchCheckedOuts() const override { return FString(); }
+	virtual bool GetOtherBranchHeadModification(FString& HeadBranchOut, FString& ActionOut, int32& HeadChangeListOut) const override { return false; }
+	virtual bool IsCurrent() const override;
+	virtual bool IsSourceControlled() const override;
+	virtual bool IsAdded() const override;
+	virtual bool IsDeleted() const override;
+	virtual bool IsIgnored() const override;
+	virtual bool CanEdit() const override;
+	virtual bool IsUnknown() const override;
+	virtual bool IsModified() const override;
+	virtual bool CanAdd() const override;
+	virtual bool CanDelete() const override;
+	virtual bool IsConflicted() const override;
+	virtual bool CanRevert() const override;
+
+public:
+	/** Filename on disk */
+	FString Filename;
+
+	/** The timestamp of the last update */
+	FDateTime TimeStamp;
+
+	/** State of the working copy */
+	ESkeinState::Type State;
+};
\ No newline at end of file
diff --git a/Engine/Plugins/Developer/SkeinSourceControl/Source/SkeinSourceControl/Private/SkeinSourceControlThumbnail.cpp b/Engine/Plugins/Developer/SkeinSourceControl/Source/SkeinSourceControl/Private/SkeinSourceControlThumbnail.cpp
new file mode 100644
--- /dev/null
+++ b/Engine/Plugins/Developer/SkeinSourceControl/Source/SkeinSourceControl/Private/SkeinSourceControlThumbnail.cpp
@@ -0,0 +1,87 @@
+// Copyright Epic Games, Inc. All Rights Reserved.
+
+#include "SkeinSourceControlThumbnail.h"
+#include "AssetRegistry/AssetRegistryModule.h"
+#include "AssetRegistry/IAssetRegistry.h"
+#include "Misc/FileHelper.h"
+#include "ISourceControlModule.h"
+#include "ObjectTools.h"
+#include "ImageUtils.h"
+
+namespace SkeinSourceControlThumbnail
+{
+
+bool WriteThumbnailToDisk(const FString& InAssetPath, const FString& InThumbnailPath, int32 InSize /* = 256*/)
+{
+	FString AssetPackageName = FPackageName::FilenameToLongPackageName(InAssetPath);
+
+	FAssetRegistryModule& AssetRegistryModule = FModuleManager::Get().LoadModuleChecked<FAssetRegistryModule>(TEXT("AssetRegistry"));
+	IAssetRegistry& AssetRegistry = AssetRegistryModule.Get();
+		
+	// Make sure to call this function using bIncludeOnlyOnDiskAssets=true as iterating in-memory assets can only happen on the main thread and
+	// this function gets called by the Skein plugin from a worker thread
+	TArray<FAssetData> AssetDatas;
+	if (!AssetRegistry.GetAssetsByPackageName(FName(*AssetPackageName), AssetDatas, true) || AssetDatas.Num() == 0)
+	{
+#if UE_BUILD_DEBUG
+		UE_LOG(LogSourceControl, Log, TEXT("WriteThumbnailToDisk failed because no assets could be found for %s."), *InAssetPath);
+#endif
+		return false;
+	}
+
+	FName FullAssetName = FName(*(AssetDatas[0].GetFullName()));
+
+	FThumbnailMap ThumbnailMap;
+	TArray<FName> ObjectNames;
+	ObjectNames.Add(FullAssetName);
+	if (!ThumbnailTools::ConditionallyLoadThumbnailsForObjects(ObjectNames, ThumbnailMap) || ThumbnailMap.Num() == 0)
+	{
+#if UE_BUILD_DEBUG
+		UE_LOG(LogSourceControl, Log, TEXT("WriteThumbnailToDisk failed because thumbnails could not be loaded for %s."), *InAssetPath);
+#endif
+		return false;
+	}
+
+	const FObjectThumbnail* ObjectThumbnail = ThumbnailMap.Find(FullAssetName);
+	check(ObjectThumbnail != nullptr);
+
+	// Grab the - uncompressed - bytes
+	const TArray<uint8>& Bytes = ObjectThumbnail->GetUncompressedImageData();
+
+	// Convert the bytes to colors
+	TArray<FColor> Colors;
+	Colors.SetNumUninitialized(Bytes.Num() / sizeof(FColor));
+	FMemory::Memcpy(Colors.GetData(), Bytes.GetData(), Bytes.Num());
+
+	// Resize if needed
+	int32 Width = ObjectThumbnail->GetImageWidth();
+	int32 Height = ObjectThumbnail->GetImageHeight();
+	if (Width != InSize || Height != InSize)
+	{
+		TArray<FColor> ResizedColors;
+		ResizedColors.SetNum(InSize * InSize);
+
+		const bool bLinearSpace = false;
+		const bool bForceOpaqueOutput = false;
+		FImageUtils::ImageResize(Width, Height, Colors, InSize, InSize, ResizedColors, bLinearSpace, bForceOpaqueOutput);
+
+		Colors = MoveTemp(ResizedColors);
+	}
+
+	// Compress the image data
+	TArray64<uint8> CompressedBitmap;
+	FImageUtils::PNGCompressImageArray(ObjectThumbnail->GetImageWidth(), ObjectThumbnail->GetImageHeight(), TArrayView64<const FColor>(Colors.GetData(), Colors.Num()), CompressedBitmap);
+		
+	// Write to disk as a PNG to maintain transparency		
+	if (!FFileHelper::SaveArrayToFile(CompressedBitmap, *InThumbnailPath))
+	{
+#if UE_BUILD_DEBUG
+		UE_LOG(LogSourceControl, Log, TEXT("WriteThumbnailToDisk failed because image data could not be written to %s."), *InThumbnailPath);
+#endif
+		return false;
+	}
+
+	return true;
+}
+
+}
\ No newline at end of file
diff --git a/Engine/Plugins/Developer/SkeinSourceControl/Source/SkeinSourceControl/Private/SkeinSourceControlThumbnail.h b/Engine/Plugins/Developer/SkeinSourceControl/Source/SkeinSourceControl/Private/SkeinSourceControlThumbnail.h
new file mode 100644
--- /dev/null
+++ b/Engine/Plugins/Developer/SkeinSourceControl/Source/SkeinSourceControl/Private/SkeinSourceControlThumbnail.h
@@ -0,0 +1,10 @@
+#pragma once
+
+#include "CoreMinimal.h"
+
+namespace SkeinSourceControlThumbnail
+{
+
+bool WriteThumbnailToDisk(const FString& InAssetPath, const FString& InThumbnailPath, int32 InSize = 256);
+
+}
\ No newline at end of file
diff --git a/Engine/Plugins/Developer/SkeinSourceControl/Source/SkeinSourceControl/Private/SkeinSourceControlThumbnailCommandlet.cpp b/Engine/Plugins/Developer/SkeinSourceControl/Source/SkeinSourceControl/Private/SkeinSourceControlThumbnailCommandlet.cpp
new file mode 100644
--- /dev/null
+++ b/Engine/Plugins/Developer/SkeinSourceControl/Source/SkeinSourceControl/Private/SkeinSourceControlThumbnailCommandlet.cpp
@@ -0,0 +1,88 @@
+// Copyright Epic Games, Inc. All Rights Reserved.
+
+#include "SkeinSourceControlThumbnailCommandlet.h"
+#include "SkeinSourceControlThumbnail.h"
+
+DEFINE_LOG_CATEGORY_STATIC(LogSkeinThumbnailCommandlet, Log, All);
+
+int32 USkeinSourceControlThumbnailCommandlet::Main(FString const& Params)
+{
+	TArray<FString> Tokens;
+	TArray<FString> Switches;
+	ParseCommandLine(*Params, Tokens, Switches);
+
+	int32 Size = 256;
+	TArray<FString> Files;
+
+	for (const FString& Switch : Switches)
+	{
+		FString FilesSwitch;
+		FString FileCountsStr;
+		if (Switch.Split(TEXT("Files="), &FilesSwitch, &FileCountsStr) || Switch.Split(TEXT("f="), &FilesSwitch, &FileCountsStr))
+		{
+			TArray<FString> FileCountStrs;
+			if (FileCountsStr.ParseIntoArray(FileCountStrs, TEXT(",")) > 0)
+			{
+				Files = MoveTemp(FileCountStrs);
+			}
+		}
+		else if (FParse::Value(*Switch, TEXT("Size="), Size) || FParse::Value(*Switch, TEXT("s="), Size))
+		{
+			continue;
+		}
+	}
+
+	if (Files.IsEmpty())
+	{
+		UE_LOG(LogSkeinThumbnailCommandlet, Warning, TEXT("You must specify one or more asset files using -Files=path/to/file1.uasset,path/to/file2.uasset"));
+		return -1;
+	}
+
+	int32 Failed = 0;
+	int32 Succeeded = 0;
+
+	for (const FString& FileName : Files)
+	{
+		if (FPaths::FileExists(FileName))
+		{
+			FString OutputFolder;
+			if (Tokens.Num() != 0)
+			{
+				OutputFolder = Tokens[0];
+			}
+			if (OutputFolder.IsEmpty())
+			{
+				OutputFolder = FPaths::GetPath(FileName);
+			}
+			FPaths::ConvertRelativePathToFull(OutputFolder);
+			FPaths::NormalizeDirectoryName(OutputFolder);
+
+			TStringBuilder<512> OutputFilenameBuilder;
+			OutputFilenameBuilder << OutputFolder;
+			OutputFilenameBuilder << FPlatformMisc::GetDefaultPathSeparator();
+			OutputFilenameBuilder << FPaths::GetBaseFilename(FileName);
+			OutputFilenameBuilder << TEXT(".png");
+
+			FString OutputFilename = OutputFilenameBuilder.ToString();
+			FPaths::NormalizeFilename(OutputFilename);
+
+			if (SkeinSourceControlThumbnail::WriteThumbnailToDisk(FileName, OutputFilename, Size))
+			{
+				UE_LOG(LogSkeinThumbnailCommandlet, Display, TEXT("Written %dx%d output: %s"), Size, Size, *OutputFilename);
+				++Succeeded;
+			}
+			else
+			{
+				UE_LOG(LogSkeinThumbnailCommandlet, Error, TEXT("Could not write %dx%d output file: %s"), Size, Size, *OutputFilename);
+				++Failed;
+			}
+		}
+		else
+		{
+			UE_LOG(LogSkeinThumbnailCommandlet, Error, TEXT("Could not find input file: %s"), *FileName);
+			++Failed;
+		}
+	}
+
+	return Failed;
+}
diff --git a/Engine/Plugins/Developer/SkeinSourceControl/Source/SkeinSourceControl/Private/SkeinSourceControlThumbnailCommandlet.h b/Engine/Plugins/Developer/SkeinSourceControl/Source/SkeinSourceControl/Private/SkeinSourceControlThumbnailCommandlet.h
new file mode 100644
--- /dev/null
+++ b/Engine/Plugins/Developer/SkeinSourceControl/Source/SkeinSourceControl/Private/SkeinSourceControlThumbnailCommandlet.h
@@ -0,0 +1,19 @@
+// Copyright Epic Games, Inc. All Rights Reserved.
+
+#pragma once
+
+#include "CoreMinimal.h"
+#include "UObject/ObjectMacros.h"
+#include "Commandlets/Commandlet.h"
+#include "SkeinSourceControlThumbnailCommandlet.generated.h"
+
+UCLASS()
+class USkeinSourceControlThumbnailCommandlet : public UCommandlet
+{
+	GENERATED_BODY()
+
+public:
+	//~ Begin UCommandlet Interface
+	virtual int32 Main(FString const& Params) override;
+	//~ End UCommandlet Interface
+};
\ No newline at end of file
diff --git a/Engine/Plugins/Developer/SkeinSourceControl/Source/SkeinSourceControl/Private/SkeinSourceControlUtils.cpp b/Engine/Plugins/Developer/SkeinSourceControl/Source/SkeinSourceControl/Private/SkeinSourceControlUtils.cpp
new file mode 100644
--- /dev/null
+++ b/Engine/Plugins/Developer/SkeinSourceControl/Source/SkeinSourceControl/Private/SkeinSourceControlUtils.cpp
@@ -0,0 +1,202 @@
+// Copyright Epic Games, Inc. All Rights Reserved.
+
+#include "SkeinSourceControlUtils.h"
+#include "ISourceControlModule.h"
+
+namespace SkeinSourceControlConstants
+{
+	/** The maximum number of files we submit in a single Skein command */
+	const int32 MaxFilesPerBatch = 50;
+}
+
+namespace SkeinSourceControlUtils
+{
+
+// Launch the Skein command line process and extract its results & errors
+static bool RunCommandInternalRaw(const FString& InCommand, const FString& InSkeinBinaryPath, const FString& InSkeinProjectRoot, const TArray<FString>& InParameters, const TArray<FString>& InFiles, FString& OutResults, FString& OutErrors)
+{
+	int32 ReturnCode = 0;
+
+	FString FullCommand;
+	FString LoggableCommand; // short version of the command for logging purpose
+
+	// Append the Skein command itself
+	LoggableCommand += InCommand;
+
+	// Append to the command all parameters, and then finally the files
+	for(const auto& Parameter : InParameters)
+	{
+		LoggableCommand += TEXT(" ");
+		LoggableCommand += Parameter;
+	}
+	for(const auto& File : InFiles)
+	{
+		LoggableCommand += TEXT(" \"");
+		LoggableCommand += File;
+		LoggableCommand += TEXT("\"");
+	}
+
+	FullCommand += LoggableCommand;
+
+#if UE_BUILD_DEBUG
+	UE_LOG(LogSourceControl, Log, TEXT("RunCommandInternalRaw: 'skein %s'"), *LoggableCommand);
+#endif
+
+	FPlatformProcess::ExecProcess(*InSkeinBinaryPath, *FullCommand, &ReturnCode, &OutResults, &OutErrors);
+
+#if UE_BUILD_DEBUG
+
+	if (OutResults.IsEmpty())
+	{
+		UE_LOG(LogSourceControl, Log, TEXT("RunCommandInternalRaw: 'OutResults=n/a'"));
+	}
+	else
+	{
+		UE_LOG(LogSourceControl, Log, TEXT("RunCommandInternalRaw(%s): OutResults=\n%s"), *InCommand, *OutResults);
+	}
+
+	if(ReturnCode != 0)
+	{
+		if (OutErrors.IsEmpty())
+		{
+			UE_LOG(LogSourceControl, Warning, TEXT("RunCommandInternalRaw: 'OutErrors=n/a'"));
+		}
+		else
+		{
+			UE_LOG(LogSourceControl, Warning, TEXT("RunCommandInternalRaw(%s): OutErrors=\n%s"), *InCommand, *OutErrors);
+		}
+	}
+#endif
+
+	return ReturnCode == 0;
+}
+
+// Basic parsing or results & errors from the Skein command line process
+static bool RunCommandInternal(const FString& InCommand, const FString& InSkeinBinaryPath, const FString& InSkeinProjectRoot, const TArray<FString>& InParameters, const TArray<FString>& InFiles, TArray<FString>& OutResults, TArray<FString>& OutErrorMessages)
+{
+	bool bResult;
+	FString Results;
+	FString Errors;
+
+	bResult = RunCommandInternalRaw(InCommand, InSkeinBinaryPath, InSkeinProjectRoot, InParameters, InFiles, Results, Errors);
+	Results.ParseIntoArray(OutResults, TEXT("\n"), true);
+	Errors.ParseIntoArray(OutErrorMessages, TEXT("\n"), true);
+
+	return bResult;
+}
+
+}
+
+namespace SkeinSourceControlUtils
+{
+
+FString FindSkeinBinaryPath()
+{
+	FString SkeinBinaryPath;
+
+#if PLATFORM_WINDOWS
+	SkeinBinaryPath = FPaths::EngineDir() / TEXT("Binaries") / TEXT("Win64") / TEXT("Skein.exe");
+#elif PLATFORM_LINUX
+	SkeinBinaryPath = FPaths::EngineDir() / TEXT("Binaries") / TEXT("Linux") / TEXT("Skein");
+#elif PLATFORM_MAC
+	SkeinBinaryPath = FPaths::EngineDir() / TEXT("Binaries") / TEXT("Mac") / TEXT("Skein");
+#endif
+
+	bool bBinaryExists = FPlatformFileManager::Get().GetPlatformFile().FileExists(*SkeinBinaryPath);
+	if (bBinaryExists)
+		return SkeinBinaryPath;
+
+	return FString();
+}
+
+FString FindSkeinProjectRoot(const FString& InPath)
+{
+	FString PathToSkeinProjectRoot = InPath;
+	FPaths::NormalizeDirectoryName(PathToSkeinProjectRoot);
+
+	while (!PathToSkeinProjectRoot.IsEmpty())
+	{
+		// Look for the "skein.yml" file present at the root of every Skein project
+		FString PathToSkeinFile = PathToSkeinProjectRoot / TEXT("skein.yml");
+		
+		bool bFound = IFileManager::Get().FileExists(*PathToSkeinFile);
+		if (bFound)
+		{
+			// Found it!
+			break;
+		}
+		else
+		{
+			int32 LastSlashIndex;
+			if (PathToSkeinProjectRoot.FindLastChar('/', LastSlashIndex))
+			{
+				PathToSkeinProjectRoot.LeftInline(LastSlashIndex);
+			}
+			else
+			{
+				PathToSkeinProjectRoot.Empty();
+			}
+		}
+	}
+
+	return PathToSkeinProjectRoot;
+}
+
+bool IsSkeinAvailable()
+{
+	FString SkeinBinaryPath = FindSkeinBinaryPath();
+	if (SkeinBinaryPath.IsEmpty())
+		return false;
+
+	return true;
+}
+
+bool IsSkeinProjectFound(const FString& InPath, FString& OutProjectRoot, FString& OutProjectName)
+{
+	OutProjectRoot = FindSkeinProjectRoot(InPath);
+	if (OutProjectRoot.IsEmpty())
+		return false;
+	
+	// SKEIN_TODO: grab project name from skein.yml
+	OutProjectName = "SKEIN_TODO";
+
+	return true;
+}
+
+bool RunCommand(const FString& InCommand, const FString& InSkeinBinaryPath, const FString& InSkeinProjectRoot, const TArray<FString>& InParameters, const TArray<FString>& InFiles, TArray<FString>& OutResults, TArray<FString>& OutErrorMessages)
+{
+	bool bResult = true;
+
+	if (InFiles.Num() > SkeinSourceControlConstants::MaxFilesPerBatch)
+	{
+		// Batch files up so we dont exceed command-line limits
+		int32 FileCount = 0;
+		while (FileCount < InFiles.Num())
+		{
+			TArray<FString> FilesInBatch;
+			for (int32 FileIndex = 0; FileCount < InFiles.Num() && FileIndex < SkeinSourceControlConstants::MaxFilesPerBatch; FileIndex++, FileCount++)
+			{
+				FilesInBatch.Add(InFiles[FileCount]);
+			}
+
+			TArray<FString> BatchResults;
+			TArray<FString> BatchErrors;
+			bResult &= RunCommandInternal(InCommand, InSkeinBinaryPath, InSkeinProjectRoot, InParameters, FilesInBatch, BatchResults, BatchErrors);
+			OutResults += BatchResults;
+			OutErrorMessages += BatchErrors;
+		}
+	}
+	else
+	{
+		bResult &= RunCommandInternal(InCommand, InSkeinBinaryPath, InSkeinProjectRoot, InParameters, InFiles, OutResults, OutErrorMessages);
+	}
+
+	return bResult;
+}
+	
+bool RunUpdateStatus(const FString& InSkeinBinaryPath, const FString& InSkeinProjectRoot, const TArray<FString>& InFiles, TArray<FString>& OutErrorMessages, TArray<FSkeinSourceControlState>& OutStates)
+{
+	return false;
+}
+
+}
\ No newline at end of file
diff --git a/Engine/Plugins/Developer/SkeinSourceControl/Source/SkeinSourceControl/Private/SkeinSourceControlUtils.h b/Engine/Plugins/Developer/SkeinSourceControl/Source/SkeinSourceControl/Private/SkeinSourceControlUtils.h
new file mode 100644
--- /dev/null
+++ b/Engine/Plugins/Developer/SkeinSourceControl/Source/SkeinSourceControl/Private/SkeinSourceControlUtils.h
@@ -0,0 +1,66 @@
+// Copyright Epic Games, Inc. All Rights Reserved.
+
+#pragma once
+
+#include "CoreMinimal.h"
+#include "SkeinSourceControlState.h"
+
+namespace SkeinSourceControlUtils
+{
+
+/**
+ * Determines the path to the Skein executable (if any)
+ * @returns an empty string if the Skein binary couldn't be found
+ */
+FString FindSkeinBinaryPath();
+
+/**
+ * Determines the path to the Skein project root (if any)
+ * @returns an empty string if the directory is not part of a Skein project
+ */
+FString FindSkeinProjectRoot(const FString& InPath);
+
+/**
+ * Finds out if the Skein environment is available.
+ * In practice this checks if the Skein CLI application is installed at the expected location.
+ * @returns true if the Skein environment can be used for source control operations.
+ */
+bool IsSkeinAvailable();
+
+/**
+  * Finds out if the given directory is part of a Skein project
+  *
+  * @param	 InPath				The directory to check
+  * @param	 OutProjectRoot		The Skein project root directory found (if any)
+  * @param	 OutProjectName		The Skein project name found (if any)
+  * @returns true if a Skein project was found.
+  */
+bool IsSkeinProjectFound(const FString& InPath, FString& OutProjectRoot, FString& OutProjectName);
+
+/**
+ * Run a Skein command - output is a string TArray
+ *
+ * @param	InCommand			The Skein command - e.g. add
+ * @param	InSkeinBinaryPath	The path to the Skein binary
+ * @param	InSkeinProjectRoot	The Skein project root from where to run the command
+ * @param	InParameters		The parameters to the Skein command
+ * @param	InFiles				The files to be operated on
+ * @param	OutResults			The results (from StdOut) as an array per-line
+ * @param	OutErrorMessages	Any errors (from StdErr) as an array per-line
+ * @returns true if the command succeeded and returned no errors
+ */
+bool RunCommand(const FString& InCommand, const FString& InSkeinBinaryPath, const FString& InSkeinProjectRoot, const TArray<FString>& InParameters, const TArray<FString>& InFiles, TArray<FString>& OutResults, TArray<FString>& OutErrorMessages);
+	
+/**
+ * Run a Skein "status" command and parse it
+ *
+ * @param	InSkeinBinaryPath	The path to the Skein binary
+ * @param	InSkeinProjectRoot	The Skein project root from where to run the command
+ * @param	InFiles				The files to be operated on
+ * @param	OutErrorMessages	Any errors (from StdErr) as an array per-line
+ * @param   OutStates           The state of each of the input files
+ * @returns true if the command succeeded and returned no errors
+ */
+bool RunUpdateStatus(const FString& InSkeinBinaryPath, const FString& InSkeinProjectRoot, const TArray<FString>& InFiles, TArray<FString>& OutErrorMessages, TArray<FSkeinSourceControlState>& OutStates);
+
+}
\ No newline at end of file
diff --git a/Engine/Plugins/Developer/SkeinSourceControl/Source/SkeinSourceControl/SkeinSourceControl.Build.cs b/Engine/Plugins/Developer/SkeinSourceControl/Source/SkeinSourceControl/SkeinSourceControl.Build.cs
new file mode 100644
--- /dev/null
+++ b/Engine/Plugins/Developer/SkeinSourceControl/Source/SkeinSourceControl/SkeinSourceControl.Build.cs
@@ -0,0 +1,24 @@
+// Copyright Epic Games, Inc. All Rights Reserved.
+
+using UnrealBuildTool;
+
+public class SkeinSourceControl : ModuleRules
+{
+	public SkeinSourceControl(ReadOnlyTargetRules Target) : base(Target)
+	{
+        PrivateDependencyModuleNames.AddRange(
+            new string[] {
+				"Core",
+				"CoreUObject",
+				"InputCore",
+				"Slate",
+				"SlateCore",
+                "EditorStyle",
+				"Engine",
+				"SourceControl",
+				"AssetRegistry",
+				"UnrealEd"
+			}
+		);
+	}
+}
diff --git a/Engine/Source/Developer/SourceControl/Private/SSourceControlPicker.cpp b/Engine/Source/Developer/SourceControl/Private/SSourceControlPicker.cpp
--- a/Engine/Source/Developer/SourceControl/Private/SSourceControlPicker.cpp
+++ b/Engine/Source/Developer/SourceControl/Private/SSourceControlPicker.cpp
@@ -117,11 +117,15 @@ FText SSourceControlPicker::GetProviderText(const FName& InName) const
 		return LOCTEXT("NoProviderDescription", "None  (source control disabled)");
 	}
 
-	// @todo: Remove this block after the Git plugin has been exhaustively tested (also remember to change the Git plugin's "IsBetaVersion" setting to false.)
-	if(InName == "Git" )
+	// @todo: Remove this block after the Git/Skein plugins have been exhaustively tested (also remember to change the Git plugin's "IsBetaVersion" setting to false.)
+	if(InName == "Git")
 	{
 		return LOCTEXT( "GitBetaProviderName", "Git  (beta version)" );
 	}
+	if(InName == "Skein")
+	{
+		return LOCTEXT("SkeinBetaProviderName", "Skein  (beta version)");
+	}
 
 	return FText::FromName(InName);
 }
