SkeinUE: Compile fix for error introduced in #17718126. Why did this compile locally?

diff --git a/Engine/Plugins/Developer/SkeinSourceControl/Source/SkeinSourceControl/Private/SkeinSourceControlState.h b/Engine/Plugins/Developer/SkeinSourceControl/Source/SkeinSourceControl/Private/SkeinSourceControlState.h
--- a/Engine/Plugins/Developer/SkeinSourceControl/Source/SkeinSourceControl/Private/SkeinSourceControlState.h
+++ b/Engine/Plugins/Developer/SkeinSourceControl/Source/SkeinSourceControl/Private/SkeinSourceControlState.h
@@ -2,6 +2,9 @@
 
 #pragma once
 
+#include "CoreMinimal.h"
+#include "ISourceControlState.h"
+
 namespace ESkeinState
 {
 	// SKEIN_TODO: Re-evaluate each of the possible states used here. Some of them might not be applicable to Skein.
