SkeinUE: Compile fix for error introduced in #17718126 regarding missing copyright boilerplate.

diff --git a/Engine/Plugins/Developer/SkeinSourceControl/Source/SkeinSourceControl/Private/SkeinSourceControlOperations.h b/Engine/Plugins/Developer/SkeinSourceControl/Source/SkeinSourceControl/Private/SkeinSourceControlOperations.h
--- a/Engine/Plugins/Developer/SkeinSourceControl/Source/SkeinSourceControl/Private/SkeinSourceControlOperations.h
+++ b/Engine/Plugins/Developer/SkeinSourceControl/Source/SkeinSourceControl/Private/SkeinSourceControlOperations.h
@@ -1,3 +1,5 @@
+// Copyright Epic Games, Inc. All Rights Reserved.
+
 #pragma once
 
 #include "CoreMinimal.h"
diff --git a/Engine/Plugins/Developer/SkeinSourceControl/Source/SkeinSourceControl/Private/SkeinSourceControlThumbnail.h b/Engine/Plugins/Developer/SkeinSourceControl/Source/SkeinSourceControl/Private/SkeinSourceControlThumbnail.h
--- a/Engine/Plugins/Developer/SkeinSourceControl/Source/SkeinSourceControl/Private/SkeinSourceControlThumbnail.h
+++ b/Engine/Plugins/Developer/SkeinSourceControl/Source/SkeinSourceControl/Private/SkeinSourceControlThumbnail.h
@@ -1,3 +1,5 @@
+// Copyright Epic Games, Inc. All Rights Reserved.
+
 #pragma once
 
 #include "CoreMinimal.h"
