SkeinUE: assorted changes:
- Added build dependency to Json.
- Fixed casing of 'Skein.exe' vs 'skein.exe' as that can cause issues on Linux/Mac.
- Fixed the Skein executable not running from within the ProjectRootDirectory, which triggers a 'Skein could not find project directory.' error. And ensured it's an absolute path.

diff --git a/Engine/Plugins/Developer/SkeinSourceControl/Source/SkeinSourceControl/Private/SkeinSourceControlUtils.cpp b/Engine/Plugins/Developer/SkeinSourceControl/Source/SkeinSourceControl/Private/SkeinSourceControlUtils.cpp
--- a/Engine/Plugins/Developer/SkeinSourceControl/Source/SkeinSourceControl/Private/SkeinSourceControlUtils.cpp
+++ b/Engine/Plugins/Developer/SkeinSourceControl/Source/SkeinSourceControl/Private/SkeinSourceControlUtils.cpp
@@ -42,7 +42,7 @@ static bool RunCommandInternalRaw(const FString& InCommand, const FString& InSke
 	UE_LOG(LogSourceControl, Log, TEXT("RunCommandInternalRaw: 'skein %s'"), *LoggableCommand);
 #endif
 
-	FPlatformProcess::ExecProcess(*InSkeinBinaryPath, *FullCommand, &ReturnCode, &OutResults, &OutErrors);
+	FPlatformProcess::ExecProcess(*InSkeinBinaryPath, *FullCommand, &ReturnCode, &OutResults, &OutErrors, *InSkeinProjectRoot);
 
 #if UE_BUILD_DEBUG
 
@@ -95,11 +95,11 @@ FString FindSkeinBinaryPath()
 	FString SkeinBinaryPath;
 
 #if PLATFORM_WINDOWS
-	SkeinBinaryPath = FPaths::EngineDir() / TEXT("Binaries") / TEXT("Win64") / TEXT("Skein.exe");
+	SkeinBinaryPath = FPaths::EngineDir() / TEXT("Binaries") / TEXT("Win64") / TEXT("skein.exe");
 #elif PLATFORM_LINUX
-	SkeinBinaryPath = FPaths::EngineDir() / TEXT("Binaries") / TEXT("Linux") / TEXT("Skein");
+	SkeinBinaryPath = FPaths::EngineDir() / TEXT("Binaries") / TEXT("Linux") / TEXT("skein");
 #elif PLATFORM_MAC
-	SkeinBinaryPath = FPaths::EngineDir() / TEXT("Binaries") / TEXT("Mac") / TEXT("Skein");
+	SkeinBinaryPath = FPaths::EngineDir() / TEXT("Binaries") / TEXT("Mac") / TEXT("skein");
 #endif
 
 	bool bBinaryExists = FPlatformFileManager::Get().GetPlatformFile().FileExists(*SkeinBinaryPath);
@@ -111,7 +111,7 @@ FString FindSkeinBinaryPath()
 
 FString FindSkeinProjectRoot(const FString& InPath)
 {
-	FString PathToSkeinProjectRoot = InPath;
+	FString PathToSkeinProjectRoot = FPaths::ConvertRelativePathToFull(InPath);
 	FPaths::NormalizeDirectoryName(PathToSkeinProjectRoot);
 
 	while (!PathToSkeinProjectRoot.IsEmpty())
diff --git a/Engine/Plugins/Developer/SkeinSourceControl/Source/SkeinSourceControl/SkeinSourceControl.Build.cs b/Engine/Plugins/Developer/SkeinSourceControl/Source/SkeinSourceControl/SkeinSourceControl.Build.cs
--- a/Engine/Plugins/Developer/SkeinSourceControl/Source/SkeinSourceControl/SkeinSourceControl.Build.cs
+++ b/Engine/Plugins/Developer/SkeinSourceControl/Source/SkeinSourceControl/SkeinSourceControl.Build.cs
@@ -17,7 +17,8 @@ public class SkeinSourceControl : ModuleRules
 				"Engine",
 				"SourceControl",
 				"AssetRegistry",
-				"UnrealEd"
+				"UnrealEd",
+				"Json",
 			}
 		);
 	}
