SkeinUE: Compile fix for NonUnity error reported by Horde.

diff --git a/Engine/Plugins/Developer/SkeinSourceControl/Source/SkeinSourceControl/Private/SkeinSourceControlModule.cpp b/Engine/Plugins/Developer/SkeinSourceControl/Source/SkeinSourceControl/Private/SkeinSourceControlModule.cpp
--- a/Engine/Plugins/Developer/SkeinSourceControl/Source/SkeinSourceControl/Private/SkeinSourceControlModule.cpp
+++ b/Engine/Plugins/Developer/SkeinSourceControl/Source/SkeinSourceControl/Private/SkeinSourceControlModule.cpp
@@ -1,6 +1,7 @@
 // Copyright Epic Games, Inc. All Rights Reserved.
 
 #include "SkeinSourceControlModule.h"
+#include "Modules/ModuleManager.h"
 #include "SkeinSourceControlOperations.h"
 #include "Features/IModularFeatures.h"
 
