SkeinUE: Compile fix for NonUnity builds.

diff --git a/Engine/Plugins/Developer/SkeinSourceControl/Source/SkeinSourceControl/Private/SSkeinSourceControlSettings.cpp b/Engine/Plugins/Developer/SkeinSourceControl/Source/SkeinSourceControl/Private/SSkeinSourceControlSettings.cpp
--- a/Engine/Plugins/Developer/SkeinSourceControl/Source/SkeinSourceControl/Private/SSkeinSourceControlSettings.cpp
+++ b/Engine/Plugins/Developer/SkeinSourceControl/Source/SkeinSourceControl/Private/SSkeinSourceControlSettings.cpp
@@ -2,6 +2,9 @@
 
 #include "SSkeinSourceControlSettings.h"
 #include "SkeinSourceControlModule.h"
+#include "Modules/ModuleManager.h"
+#include "Widgets/SBoxPanel.h"
+#include "Widgets/Text/STextBlock.h"
 
 #define LOCTEXT_NAMESPACE "SSkeinSourceControlSettings"
 
diff --git a/Engine/Plugins/Developer/SkeinSourceControl/Source/SkeinSourceControl/Private/SkeinSourceControlUtils.cpp b/Engine/Plugins/Developer/SkeinSourceControl/Source/SkeinSourceControl/Private/SkeinSourceControlUtils.cpp
--- a/Engine/Plugins/Developer/SkeinSourceControl/Source/SkeinSourceControl/Private/SkeinSourceControlUtils.cpp
+++ b/Engine/Plugins/Developer/SkeinSourceControl/Source/SkeinSourceControl/Private/SkeinSourceControlUtils.cpp
@@ -1,6 +1,10 @@
 // Copyright Epic Games, Inc. All Rights Reserved.
 
 #include "SkeinSourceControlUtils.h"
+#include "HAL/PlatformProcess.h"
+#include "HAL/PlatformFileManager.h"
+#include "HAL/FileManager.h"
+#include "Misc/Paths.h"
 #include "ISourceControlModule.h"
 
 namespace SkeinSourceControlConstants
