SkeinUE: Parse SkeinCLI output as JSON.
- Also implements parsing of the 'skein status' command assuming two columns ('path' and 'status'). This can be easily modified to suit whatever the CLI ends up outputting as that's still under development.
- Also adds the Revert operation to the Skein source control plugin.
- Also fixed some NonUnity compile errors.

diff --git a/Engine/Plugins/Developer/SkeinSourceControl/Source/SkeinSourceControl/Private/SkeinSourceControlCommand.h b/Engine/Plugins/Developer/SkeinSourceControl/Source/SkeinSourceControl/Private/SkeinSourceControlCommand.h
--- a/Engine/Plugins/Developer/SkeinSourceControl/Source/SkeinSourceControl/Private/SkeinSourceControlCommand.h
+++ b/Engine/Plugins/Developer/SkeinSourceControl/Source/SkeinSourceControl/Private/SkeinSourceControlCommand.h
@@ -69,7 +69,7 @@ public:
 	/** Files to perform this operation on */
 	TArray<FString> Files;
 
-	/** Info and/or warning message message storage */
+	/** Info and/or warning message storage */
 	TArray<FString> InfoMessages;
 
 	/** Potential error message storage */
diff --git a/Engine/Plugins/Developer/SkeinSourceControl/Source/SkeinSourceControl/Private/SkeinSourceControlModule.cpp b/Engine/Plugins/Developer/SkeinSourceControl/Source/SkeinSourceControl/Private/SkeinSourceControlModule.cpp
--- a/Engine/Plugins/Developer/SkeinSourceControl/Source/SkeinSourceControl/Private/SkeinSourceControlModule.cpp
+++ b/Engine/Plugins/Developer/SkeinSourceControl/Source/SkeinSourceControl/Private/SkeinSourceControlModule.cpp
@@ -20,10 +20,10 @@ void FSkeinSourceControlModule::StartupModule()
 	SkeinSourceControlProvider.RegisterWorker("CheckIn", FGetSkeinSourceControlWorker::CreateStatic(&CreateWorker<FSkeinCheckInWorker>));
 	SkeinSourceControlProvider.RegisterWorker("MarkForAdd", FGetSkeinSourceControlWorker::CreateStatic(&CreateWorker<FSkeinMarkForAddWorker>));
 	SkeinSourceControlProvider.RegisterWorker("Delete", FGetSkeinSourceControlWorker::CreateStatic(&CreateWorker<FSkeinDeleteWorker>));
+	SkeinSourceControlProvider.RegisterWorker("Revert", FGetSkeinSourceControlWorker::CreateStatic(&CreateWorker<FSkeinRevertWorker>));
 	SkeinSourceControlProvider.RegisterWorker("Sync", FGetSkeinSourceControlWorker::CreateStatic(&CreateWorker<FSkeinSyncWorker>));
 	SkeinSourceControlProvider.RegisterWorker("UpdateStatus", FGetSkeinSourceControlWorker::CreateStatic(&CreateWorker<FSkeinUpdateStatusWorker>));
 
-//	SkeinSourceControlProvider.RegisterWorker("Revert", FGetSkeinSourceControlWorker::CreateStatic(&CreateWorker<FSkeinRevertWorker>));
 //	SkeinSourceControlProvider.RegisterWorker("Copy", FGetSkeinSourceControlWorker::CreateStatic(&CreateWorker<FSkeinCopyWorker>));
 //	SkeinSourceControlProvider.RegisterWorker("Resolve", FGetSkeinSourceControlWorker::CreateStatic(&CreateWorker<FSkeinResolveWorker>));
 
diff --git a/Engine/Plugins/Developer/SkeinSourceControl/Source/SkeinSourceControl/Private/SkeinSourceControlOperations.cpp b/Engine/Plugins/Developer/SkeinSourceControl/Source/SkeinSourceControl/Private/SkeinSourceControlOperations.cpp
--- a/Engine/Plugins/Developer/SkeinSourceControl/Source/SkeinSourceControl/Private/SkeinSourceControlOperations.cpp
+++ b/Engine/Plugins/Developer/SkeinSourceControl/Source/SkeinSourceControl/Private/SkeinSourceControlOperations.cpp
@@ -7,6 +7,7 @@
 #include "SkeinSourceControlModule.h"
 #include "SkeinSourceControlUtils.h"
 #include "SkeinSourceControlState.h"
+#include "Modules/ModuleManager.h"
 #include "SourceControlOperations.h"
 
 #define LOCTEXT_NAMESPACE "SkeinSourceControl"
@@ -48,17 +49,18 @@ bool FSkeinConnectWorker::Execute(FSkeinSourceControlCommand& InCommand)
 	if (InCommand.SkeinBinaryPath.IsEmpty())
 	{
 		StaticCastSharedRef<FConnect>(InCommand.Operation)->SetErrorText(LOCTEXT("SkeinNotAvailable", "The Skein Command Line application could not be found."));
-		return false;
+		InCommand.bCommandSuccessful = false;
+		return InCommand.bCommandSuccessful;
 	}
 
 	if (InCommand.SkeinProjectRoot.IsEmpty())
 	{
 		StaticCastSharedRef<FConnect>(InCommand.Operation)->SetErrorText(LOCTEXT("SkeinNotEnabled", "There is no Skein project initialized for this location."));
-		return false;
+		InCommand.bCommandSuccessful = false;
+		return InCommand.bCommandSuccessful;
 	}
 
 	InCommand.bCommandSuccessful = SkeinSourceControlUtils::RunCommand(TEXT("auth login"), InCommand.SkeinBinaryPath, InCommand.SkeinProjectRoot, TArray<FString>(), InCommand.Files, InCommand.InfoMessages, InCommand.ErrorMessages);
-
 	return InCommand.bCommandSuccessful;
 }
 
@@ -105,7 +107,7 @@ bool FSkeinMarkForAddWorker::Execute(FSkeinSourceControlCommand& InCommand)
 {
 	check(InCommand.Operation->GetName() == GetName());
 
-	InCommand.bCommandSuccessful = SkeinSourceControlUtils::RunCommand(TEXT("assets add"), InCommand.SkeinBinaryPath, InCommand.SkeinProjectRoot, TArray<FString>(), InCommand.Files, InCommand.InfoMessages, InCommand.ErrorMessages);
+	InCommand.bCommandSuccessful = SkeinSourceControlUtils::RunCommand(TEXT("assets track"), InCommand.SkeinBinaryPath, InCommand.SkeinProjectRoot, TArray<FString>(), InCommand.Files, InCommand.InfoMessages, InCommand.ErrorMessages);
 
 	SkeinSourceControlUtils::RunUpdateStatus(InCommand.SkeinBinaryPath, InCommand.SkeinProjectRoot, InCommand.Files, InCommand.ErrorMessages, States);
 
@@ -130,7 +132,7 @@ bool FSkeinDeleteWorker::Execute(FSkeinSourceControlCommand& InCommand)
 {
 	check(InCommand.Operation->GetName() == GetName());
 
-	InCommand.bCommandSuccessful = SkeinSourceControlUtils::RunCommand(TEXT("assets rm"), InCommand.SkeinBinaryPath, InCommand.SkeinProjectRoot, TArray<FString>(), InCommand.Files, InCommand.InfoMessages, InCommand.ErrorMessages);
+	InCommand.bCommandSuccessful = SkeinSourceControlUtils::RunCommand(TEXT("assets untrack"), InCommand.SkeinBinaryPath, InCommand.SkeinProjectRoot, TArray<FString>(), InCommand.Files, InCommand.InfoMessages, InCommand.ErrorMessages);
 
 	SkeinSourceControlUtils::RunUpdateStatus(InCommand.SkeinBinaryPath, InCommand.SkeinProjectRoot, InCommand.Files, InCommand.ErrorMessages, States);
 
@@ -146,6 +148,29 @@ bool FSkeinDeleteWorker::UpdateStates() const
 
 
 
+FName FSkeinRevertWorker::GetName() const
+{
+	return "Revert";
+}
+
+bool FSkeinRevertWorker::Execute(FSkeinSourceControlCommand& InCommand)
+{
+	check(InCommand.Operation->GetName() == GetName());
+
+	InCommand.bCommandSuccessful = SkeinSourceControlUtils::RunCommand(TEXT("assets revert"), InCommand.SkeinBinaryPath, InCommand.SkeinProjectRoot, TArray<FString>(), InCommand.Files, InCommand.InfoMessages, InCommand.ErrorMessages);
+
+	SkeinSourceControlUtils::RunUpdateStatus(InCommand.SkeinBinaryPath, InCommand.SkeinProjectRoot, InCommand.Files, InCommand.ErrorMessages, States);
+
+	return InCommand.bCommandSuccessful;
+}
+
+bool FSkeinRevertWorker::UpdateStates() const
+{
+	return UpdateCachedStates(States);
+}
+
+
+
 FName FSkeinSyncWorker::GetName() const
 {
 	return "Sync";
diff --git a/Engine/Plugins/Developer/SkeinSourceControl/Source/SkeinSourceControl/Private/SkeinSourceControlOperations.h b/Engine/Plugins/Developer/SkeinSourceControl/Source/SkeinSourceControl/Private/SkeinSourceControlOperations.h
--- a/Engine/Plugins/Developer/SkeinSourceControl/Source/SkeinSourceControl/Private/SkeinSourceControlOperations.h
+++ b/Engine/Plugins/Developer/SkeinSourceControl/Source/SkeinSourceControl/Private/SkeinSourceControlOperations.h
@@ -66,6 +66,21 @@ public:
 	TArray<FSkeinSourceControlState> States;
 };
 
+/** Revert any change to a file to its state on the local depot. */
+class FSkeinRevertWorker : public ISkeinSourceControlWorker
+{
+public:
+	virtual ~FSkeinRevertWorker() {}
+	/** ISkeinSourceControlWorker interface */
+	virtual FName GetName() const override;
+	virtual bool Execute(class FSkeinSourceControlCommand& InCommand) override;
+	virtual bool UpdateStates() const override;
+
+public:
+	/** Temporary states for results */
+	TArray<FSkeinSourceControlState> States;
+};
+
 /** Skein synchronize to the active snapshot set in WebUI. */
 class FSkeinSyncWorker : public ISkeinSourceControlWorker
 {
diff --git a/Engine/Plugins/Developer/SkeinSourceControl/Source/SkeinSourceControl/Private/SkeinSourceControlProvider.cpp b/Engine/Plugins/Developer/SkeinSourceControl/Source/SkeinSourceControl/Private/SkeinSourceControlProvider.cpp
--- a/Engine/Plugins/Developer/SkeinSourceControl/Source/SkeinSourceControl/Private/SkeinSourceControlProvider.cpp
+++ b/Engine/Plugins/Developer/SkeinSourceControl/Source/SkeinSourceControl/Private/SkeinSourceControlProvider.cpp
@@ -9,6 +9,7 @@
 #include "SourceControlOperations.h"
 #include "SourceControlHelpers.h"
 #include "Logging/MessageLog.h"
+#include "Misc/Paths.h"
 
 #if SOURCE_CONTROL_WITH_SLATE
 #include "SSkeinSourceControlSettings.h"
diff --git a/Engine/Plugins/Developer/SkeinSourceControl/Source/SkeinSourceControl/Private/SkeinSourceControlThumbnailCommandlet.cpp b/Engine/Plugins/Developer/SkeinSourceControl/Source/SkeinSourceControl/Private/SkeinSourceControlThumbnailCommandlet.cpp
--- a/Engine/Plugins/Developer/SkeinSourceControl/Source/SkeinSourceControl/Private/SkeinSourceControlThumbnailCommandlet.cpp
+++ b/Engine/Plugins/Developer/SkeinSourceControl/Source/SkeinSourceControl/Private/SkeinSourceControlThumbnailCommandlet.cpp
@@ -2,6 +2,7 @@
 
 #include "SkeinSourceControlThumbnailCommandlet.h"
 #include "SkeinSourceControlThumbnail.h"
+#include "Misc/Paths.h"
 
 DEFINE_LOG_CATEGORY_STATIC(LogSkeinThumbnailCommandlet, Log, All);
 
diff --git a/Engine/Plugins/Developer/SkeinSourceControl/Source/SkeinSourceControl/Private/SkeinSourceControlUtils.cpp b/Engine/Plugins/Developer/SkeinSourceControl/Source/SkeinSourceControl/Private/SkeinSourceControlUtils.cpp
--- a/Engine/Plugins/Developer/SkeinSourceControl/Source/SkeinSourceControl/Private/SkeinSourceControlUtils.cpp
+++ b/Engine/Plugins/Developer/SkeinSourceControl/Source/SkeinSourceControl/Private/SkeinSourceControlUtils.cpp
@@ -5,21 +5,99 @@
 #include "HAL/PlatformFileManager.h"
 #include "HAL/FileManager.h"
 #include "Misc/Paths.h"
+#include "Misc/ScopeLock.h"
 #include "ISourceControlModule.h"
+#include "Serialization/JsonSerializer.h"
+#include "Dom/JsonValue.h"
+#include "Misc/FileHelper.h"
 
 namespace SkeinSourceControlConstants
 {
 	/** The maximum number of files we submit in a single Skein command */
-	const int32 MaxFilesPerBatch = 50;
+	const int32 MaxFilesPerBatch = 1;
 }
 
 namespace SkeinSourceControlUtils
 {
 
+// Convert an array of FJsonObject statuses to FSkeinSourceControlStates
+static void ParseStatusOutput(const FString& InSkeinProjectRoot, const TArray<FString>& InFiles, const TSharedPtr<FJsonObject> InStates, TArray<FSkeinSourceControlState>& OutStates)
+{
+	const FDateTime Now = FDateTime::Now();
+
+	// Build map of all states received
+	TMap<FString, FString> FileStates;
+
+	const TArray<TSharedPtr<FJsonValue>>& Entries = InStates->GetArrayField("Data");
+	for (const TSharedPtr<FJsonValue>& Entry : Entries)
+	{
+		const TSharedPtr<FJsonObject>* PathStatus;
+		if (Entry->TryGetObject(PathStatus))
+		{
+			FString RelativePath = (*PathStatus)->GetStringField(TEXT("file_path"));
+			FString AbsolutePath = FPaths::Combine(InSkeinProjectRoot, RelativePath);
+			FString State = (*PathStatus)->GetStringField(TEXT("file_state"));
+
+			FPaths::NormalizeFilename(AbsolutePath);
+
+			FileStates.Add(AbsolutePath, State);
+		}
+	}
+
+	// Iterate on all files explicitly listed in the command
+	for (const auto& File : InFiles)
+	{
+		FSkeinSourceControlState FileState(File);
+		FileState.TimeStamp = Now;
+		FileState.State = ESkeinState::Unknown;
+
+		FString* State = FileStates.Find(File);
+		if (State != nullptr)
+		{
+			if (*State == "unknown")
+			{
+				FileState.State = ESkeinState::Unknown;
+			}
+			else if (*State == "add")
+			{
+				FileState.State = ESkeinState::Added;
+			}
+			else if (*State == "remove")
+			{
+				FileState.State = ESkeinState::Deleted;
+			}
+			else if (*State == "modified")
+			{
+				FileState.State = ESkeinState::Modified;
+			}
+			else if (*State == "untracked")
+			{
+				FileState.State = ESkeinState::NotControlled;
+			}
+			else if (*State == "unchanged")
+			{
+				FileState.State = ESkeinState::Unchanged;
+			}
+			else
+			{
+				checkNoEntry();
+			}
+		}
+
+		OutStates.Add(FileState);
+	}
+}
+
+// The Skein command line process returns 'instance already running' when invoked in parallel.
+static FCriticalSection RunCriticalSection;
+
 // Launch the Skein command line process and extract its results & errors
-static bool RunCommandInternalRaw(const FString& InCommand, const FString& InSkeinBinaryPath, const FString& InSkeinProjectRoot, const TArray<FString>& InParameters, const TArray<FString>& InFiles, FString& OutResults, FString& OutErrors)
+static bool RunCommandInternalRaw(const FString& InCommand, const FString& InSkeinBinaryPath, const FString& InSkeinProjectRoot, const TArray<FString>& InParameters, const TArray<FString>& InFiles, FString& OutMessage, TSharedPtr<FJsonObject>& OutData)
 {
+	FScopeLock Lock(&RunCriticalSection);
+
 	int32 ReturnCode = 0;
+	int32 StatusCode = 0;
 
 	FString FullCommand;
 	FString LoggableCommand; // short version of the command for logging purpose
@@ -27,6 +105,10 @@ static bool RunCommandInternalRaw(const FString& InCommand, const FString& InSke
 	// Append the Skein command itself
 	LoggableCommand += InCommand;
 
+	// Append the "--json" param to indicate we want Json output
+	LoggableCommand += TEXT(" ");
+	LoggableCommand += TEXT("--json");
+
 	// Append to the command all parameters, and then finally the files
 	for(const auto& Parameter : InParameters)
 	{
@@ -42,51 +124,67 @@ static bool RunCommandInternalRaw(const FString& InCommand, const FString& InSke
 
 	FullCommand += LoggableCommand;
 
-#if UE_BUILD_DEBUG
-	UE_LOG(LogSourceControl, Log, TEXT("RunCommandInternalRaw: 'skein %s'"), *LoggableCommand);
-#endif
+	UE_LOG(LogSourceControl, VeryVerbose, TEXT("RunCommandInternalRaw: 'skein %s'"), *LoggableCommand);
 
-	FPlatformProcess::ExecProcess(*InSkeinBinaryPath, *FullCommand, &ReturnCode, &OutResults, &OutErrors, *InSkeinProjectRoot);
+	FString Results;
+	FString Errors;
 
-#if UE_BUILD_DEBUG
+	FPlatformProcess::ExecProcess(*InSkeinBinaryPath, *FullCommand, &ReturnCode, &Results, &Errors, *InSkeinProjectRoot);
 
-	if (OutResults.IsEmpty())
+	if (!Results.IsEmpty())
 	{
-		UE_LOG(LogSourceControl, Log, TEXT("RunCommandInternalRaw: 'OutResults=n/a'"));
+		UE_LOG(LogSourceControl, VeryVerbose, TEXT("RunCommandInternalRaw(%s): Results=\n%s"), *InCommand, *Results);
 	}
-	else
+	if (!Errors.IsEmpty())
 	{
-		UE_LOG(LogSourceControl, Log, TEXT("RunCommandInternalRaw(%s): OutResults=\n%s"), *InCommand, *OutResults);
+		UE_LOG(LogSourceControl, VeryVerbose, TEXT("RunCommandInternalRaw(%s): Errors=\n%s"), *InCommand, *Errors);
 	}
 
-	if(ReturnCode != 0)
+	// Try to parse either the StdOut or StdErr stream as Json
+	FString OutputToParse = Results.IsEmpty() ? Errors : Results;
+	if (!OutputToParse.IsEmpty())
 	{
-		if (OutErrors.IsEmpty())
+		TSharedPtr<FJsonObject> JsonObject;
+		TSharedRef<TJsonReader<>> JsonReader = TJsonReaderFactory<>::Create(*OutputToParse);
+		if (FJsonSerializer::Deserialize(JsonReader, JsonObject) && JsonObject.IsValid())
 		{
-			UE_LOG(LogSourceControl, Warning, TEXT("RunCommandInternalRaw: 'OutErrors=n/a'"));
+			OutMessage = JsonObject->GetStringField("Message");
+			OutData = JsonObject;
+			StatusCode = JsonObject->GetIntegerField("Code");
+		}
+	}
+
+	// If unsuccessful, construct a JsonObject to return
+	if (!OutData.IsValid())
+	{
+		if (Errors.IsEmpty())
+		{
+			OutMessage = FString::Format(TEXT("Internal error ('{0}')"), { ReturnCode });
 		}
 		else
 		{
-			UE_LOG(LogSourceControl, Warning, TEXT("RunCommandInternalRaw(%s): OutErrors=\n%s"), *InCommand, *OutErrors);
+			OutMessage = Errors;
 		}
+
+		OutData = MakeShared<FJsonObject>();
+		OutData->SetBoolField(TEXT("OK"), false);
+		OutData->SetNumberField(TEXT("Code"), 500);
+		OutData->SetStringField(TEXT("Message"), OutMessage);
+		OutData->SetField(TEXT("Data"), MakeShared<FJsonValueNull>());
 	}
-#endif
 
-	return ReturnCode == 0;
+ 	return ReturnCode == 0 && StatusCode == 200;
 }
 
-// Basic parsing or results & errors from the Skein command line process
-static bool RunCommandInternal(const FString& InCommand, const FString& InSkeinBinaryPath, const FString& InSkeinProjectRoot, const TArray<FString>& InParameters, const TArray<FString>& InFiles, TArray<FString>& OutResults, TArray<FString>& OutErrorMessages)
+static bool RunCommandInternal(const FString& InCommand, const FString& InSkeinBinaryPath, const FString& InSkeinProjectRoot, const TArray<FString>& InParameters, const TArray<FString>& InFiles, FString& OutMessage, TSharedPtr<FJsonObject>& OutData)
 {
-	bool bResult;
-	FString Results;
-	FString Errors;
-
-	bResult = RunCommandInternalRaw(InCommand, InSkeinBinaryPath, InSkeinProjectRoot, InParameters, InFiles, Results, Errors);
-	Results.ParseIntoArray(OutResults, TEXT("\n"), true);
-	Errors.ParseIntoArray(OutErrorMessages, TEXT("\n"), true);
+	return RunCommandInternalRaw(InCommand, InSkeinBinaryPath, InSkeinProjectRoot, InParameters, InFiles, OutMessage, OutData);
+}
 
-	return bResult;
+static bool RunCommandInternal(const FString& InCommand, const FString& InSkeinBinaryPath, const FString& InSkeinProjectRoot, const TArray<FString>& InParameters, const TArray<FString>& InFiles, FString& OutMessage)
+{
+	TSharedPtr<FJsonObject> RawData = MakeShared<FJsonObject>();
+	return RunCommandInternalRaw(InCommand, InSkeinBinaryPath, InSkeinProjectRoot, InParameters, InFiles, OutMessage, RawData);
 }
 
 }
@@ -160,16 +258,34 @@ bool IsSkeinProjectFound(const FString& InPath, FString& OutProjectRoot, FString
 	OutProjectRoot = FindSkeinProjectRoot(InPath);
 	if (OutProjectRoot.IsEmpty())
 		return false;
-	
-	// SKEIN_TODO: grab project name from skein.yml
-	OutProjectName = "SKEIN_TODO";
+
+	FString Marker = TEXT("name:");
+	FString Filename = FPaths::Combine(OutProjectRoot, "skein.yml");
+
+	TArray<FString> Lines;
+	FFileHelper::LoadFileToStringArrayWithPredicate(Lines, *Filename,
+		[&] (const FString& Line)
+		{
+			return Line.StartsWith(Marker);
+		}
+	);
+
+	if (Lines.Num() == 1)
+	{
+		OutProjectName = Lines[0].RightChop(Marker.Len()).TrimStartAndEnd();
+	}
+	else
+	{
+		OutProjectName = "Unknown";
+	}
 
 	return true;
 }
 
-bool RunCommand(const FString& InCommand, const FString& InSkeinBinaryPath, const FString& InSkeinProjectRoot, const TArray<FString>& InParameters, const TArray<FString>& InFiles, TArray<FString>& OutResults, TArray<FString>& OutErrorMessages)
+template <typename FunctionType>
+bool RunCommandBatched(const FString& InCommand, const FString& InSkeinBinaryPath, const FString& InSkeinProjectRoot, const TArray<FString>& InParameters, const TArray<FString>& InFiles, FunctionType& OutCallback)
 {
-	bool bResult = true;
+	int NumErrors = 0;
 
 	if (InFiles.Num() > SkeinSourceControlConstants::MaxFilesPerBatch)
 	{
@@ -183,24 +299,74 @@ bool RunCommand(const FString& InCommand, const FString& InSkeinBinaryPath, cons
 				FilesInBatch.Add(InFiles[FileCount]);
 			}
 
-			TArray<FString> BatchResults;
-			TArray<FString> BatchErrors;
-			bResult &= RunCommandInternal(InCommand, InSkeinBinaryPath, InSkeinProjectRoot, InParameters, FilesInBatch, BatchResults, BatchErrors);
-			OutResults += BatchResults;
-			OutErrorMessages += BatchErrors;
+			FString BatchMessage;
+			TSharedPtr<FJsonObject> BatchData;
+			if (RunCommandInternal(InCommand, InSkeinBinaryPath, InSkeinProjectRoot, InParameters, FilesInBatch, BatchMessage, BatchData))
+			{
+				OutCallback(true, FilesInBatch, BatchMessage, BatchData);
+			}
+			else
+			{
+				OutCallback(false, FilesInBatch, BatchMessage, BatchData);
+				++NumErrors;
+			}
 		}
 	}
 	else
 	{
-		bResult &= RunCommandInternal(InCommand, InSkeinBinaryPath, InSkeinProjectRoot, InParameters, InFiles, OutResults, OutErrorMessages);
+		FString Message;
+		TSharedPtr<FJsonObject> Data;
+		if (RunCommandInternal(InCommand, InSkeinBinaryPath, InSkeinProjectRoot, InParameters, InFiles, Message, Data))
+		{
+			OutCallback(true, InFiles, Message, Data);
+		}
+		else
+		{
+			OutCallback(false, InFiles, Message, Data);
+			++NumErrors;
+		}
 	}
 
-	return bResult;
+	return (NumErrors == 0);
 }
-	
-bool RunUpdateStatus(const FString& InSkeinBinaryPath, const FString& InSkeinProjectRoot, const TArray<FString>& InFiles, TArray<FString>& OutErrorMessages, TArray<FSkeinSourceControlState>& OutStates)
+
+bool RunCommand(const FString& InCommand, const FString& InSkeinBinaryPath, const FString& InSkeinProjectRoot, const TArray<FString>& InParameters, const TArray<FString>& InFiles, TArray<FString>& OutResults, TArray<FString>& OutErrors)
 {
-	return false;
+	auto Callback = 
+		[&] (bool bBatchResult, const TArray<FString>& BatchFiles, const FString& BatchMessage, const TSharedPtr<FJsonObject>& BatchData)
+		{
+			if (bBatchResult)
+			{
+				OutResults.Add(BatchMessage);
+			}
+			else
+			{
+				OutErrors.Add(BatchMessage);
+			}
+		};
+
+	return RunCommandBatched(InCommand, InSkeinBinaryPath, InSkeinProjectRoot, InParameters, InFiles, Callback);
 }
+	
+bool RunUpdateStatus(const FString& InSkeinBinaryPath, const FString& InSkeinProjectRoot, const TArray<FString>& InFiles, TArray<FString>& OutErrors, TArray<FSkeinSourceControlState>& OutStates)
+{
+	auto Callback = 
+		[&] (bool bBatchResult, const TArray<FString>& BatchFiles, const FString& BatchMessage, const TSharedPtr<FJsonObject>& BatchData)
+		{
+			if (bBatchResult)
+			{
+				ParseStatusOutput(InSkeinProjectRoot, InFiles, BatchData, OutStates);
+			}
+			else
+			{
+				OutErrors.Add(BatchMessage);
+			}
+		};
+
+	TArray<FString> Paths;
+	Paths.Add(InSkeinProjectRoot);
 
+	return RunCommandBatched(TEXT("projects status"), InSkeinBinaryPath, InSkeinProjectRoot, TArray<FString>(), Paths, Callback);
+}
+	
 }
\ No newline at end of file
diff --git a/Engine/Plugins/Developer/SkeinSourceControl/Source/SkeinSourceControl/Private/SkeinSourceControlUtils.h b/Engine/Plugins/Developer/SkeinSourceControl/Source/SkeinSourceControl/Private/SkeinSourceControlUtils.h
--- a/Engine/Plugins/Developer/SkeinSourceControl/Source/SkeinSourceControl/Private/SkeinSourceControlUtils.h
+++ b/Engine/Plugins/Developer/SkeinSourceControl/Source/SkeinSourceControl/Private/SkeinSourceControlUtils.h
@@ -45,11 +45,11 @@ bool IsSkeinProjectFound(const FString& InPath, FString& OutProjectRoot, FString
  * @param	InSkeinProjectRoot	The Skein project root from where to run the command
  * @param	InParameters		The parameters to the Skein command
  * @param	InFiles				The files to be operated on
- * @param	OutResults			The results (from StdOut) as an array per-line
- * @param	OutErrorMessages	Any errors (from StdErr) as an array per-line
+ * @param	OutResults			The results as an array per-line
+ * @param	OutErrors			The errors as an array per-line
  * @returns true if the command succeeded and returned no errors
  */
-bool RunCommand(const FString& InCommand, const FString& InSkeinBinaryPath, const FString& InSkeinProjectRoot, const TArray<FString>& InParameters, const TArray<FString>& InFiles, TArray<FString>& OutResults, TArray<FString>& OutErrorMessages);
+bool RunCommand(const FString& InCommand, const FString& InSkeinBinaryPath, const FString& InSkeinProjectRoot, const TArray<FString>& InParameters, const TArray<FString>& InFiles, TArray<FString>& OutResults, TArray<FString>& OutErrors);
 	
 /**
  * Run a Skein "status" command and parse it
@@ -57,10 +57,9 @@ bool RunCommand(const FString& InCommand, const FString& InSkeinBinaryPath, cons
  * @param	InSkeinBinaryPath	The path to the Skein binary
  * @param	InSkeinProjectRoot	The Skein project root from where to run the command
  * @param	InFiles				The files to be operated on
- * @param	OutErrorMessages	Any errors (from StdErr) as an array per-line
+ * @param	OutErrors			Any errors as an array per-line
  * @param   OutStates           The state of each of the input files
  * @returns true if the command succeeded and returned no errors
  */
-bool RunUpdateStatus(const FString& InSkeinBinaryPath, const FString& InSkeinProjectRoot, const TArray<FString>& InFiles, TArray<FString>& OutErrorMessages, TArray<FSkeinSourceControlState>& OutStates);
-
+bool RunUpdateStatus(const FString& InSkeinBinaryPath, const FString& InSkeinProjectRoot, const TArray<FString>& InFiles, TArray<FString>& OutErrors, TArray<FSkeinSourceControlState>& OutStates);
 }
\ No newline at end of file
