SkeinUE: Handle FileState 'Missing' that was introduced in #18070806.

diff --git a/Engine/Plugins/Developer/SkeinSourceControl/Source/SkeinSourceControl/Private/SkeinSourceControlUtils.cpp b/Engine/Plugins/Developer/SkeinSourceControl/Source/SkeinSourceControl/Private/SkeinSourceControlUtils.cpp
--- a/Engine/Plugins/Developer/SkeinSourceControl/Source/SkeinSourceControl/Private/SkeinSourceControlUtils.cpp
+++ b/Engine/Plugins/Developer/SkeinSourceControl/Source/SkeinSourceControl/Private/SkeinSourceControlUtils.cpp
@@ -66,6 +66,10 @@ static void ParseStatusOutput(const FString& InSkeinProjectRoot, const TArray<FS
 			{
 				FileState.State = ESkeinState::Deleted;
 			}
+			else if (*State == "missing")
+			{
+				FileState.State = ESkeinState::Missing;
+			}
 			else if (*State == "modified")
 			{
 				FileState.State = ESkeinState::Modified;
