SkeinUE: Limit 'projects status' queries to the requested subdirectories instead of always requesting the state of all files in the project root. This is much faster as the CLI won't end up hashing all files.

diff --git a/Engine/Plugins/Developer/SkeinSourceControl/Source/SkeinSourceControl/Private/SkeinSourceControlUtils.cpp b/Engine/Plugins/Developer/SkeinSourceControl/Source/SkeinSourceControl/Private/SkeinSourceControlUtils.cpp
--- a/Engine/Plugins/Developer/SkeinSourceControl/Source/SkeinSourceControl/Private/SkeinSourceControlUtils.cpp
+++ b/Engine/Plugins/Developer/SkeinSourceControl/Source/SkeinSourceControl/Private/SkeinSourceControlUtils.cpp
@@ -367,10 +367,27 @@ bool RunUpdateStatus(const FString& InSkeinBinaryPath, const FString& InSkeinPro
 			}
 		};
 
-	TArray<FString> Paths;
-	Paths.Add(InSkeinProjectRoot);
+	TSet<FString> UniquePaths;
+	for (const auto& File : InFiles)
+	{
+		const FString Path = FPaths::GetPath(*File);
+		UniquePaths.Add(Path);
+	}
+
+	int NumErrors = 0;
+
+	for (const auto& UniquePath : UniquePaths)
+	{
+		TArray<FString> Paths;
+		Paths.Add(UniquePath);
+
+		if (!RunCommandBatched(TEXT("projects status"), InSkeinBinaryPath, InSkeinProjectRoot, TArray<FString>(), Paths, Callback))
+		{
+			NumErrors++;
+		}
+	}
 
-	return RunCommandBatched(TEXT("projects status"), InSkeinBinaryPath, InSkeinProjectRoot, TArray<FString>(), Paths, Callback);
+	return (NumErrors == 0);
 }
 	
 }
\ No newline at end of file
