SkeinUE: Extract metadata + thumbnail from the asset to disk in a directory where the CLI can find them when doing a commit.

diff --git a/Engine/Plugins/Developer/SkeinSourceControl/Source/SkeinSourceControl/Private/SkeinSourceControlMetadata.cpp b/Engine/Plugins/Developer/SkeinSourceControl/Source/SkeinSourceControl/Private/SkeinSourceControlMetadata.cpp
new file mode 100644
--- /dev/null
+++ b/Engine/Plugins/Developer/SkeinSourceControl/Source/SkeinSourceControl/Private/SkeinSourceControlMetadata.cpp
@@ -0,0 +1,239 @@
+// Copyright Epic Games, Inc. All Rights Reserved.
+
+#include "SkeinSourceControlMetadata.h"
+#include "AssetRegistry/AssetRegistryModule.h"
+#include "AssetRegistry/IAssetRegistry.h"
+#include "Serialization/JsonSerializer.h"
+#include "Misc/AssetRegistryInterface.h"
+#include "Misc/FileHelper.h"
+#include "ISourceControlModule.h"
+#include "ObjectTools.h"
+#include "ImageUtils.h"
+
+namespace SkeinSourceControlMetadata
+{
+	
+// Tags are somewhat odd. If the asset is loaded, it's expected that the Tags are gathered dynamically
+// through the UObject::GetAssetRegistryTags function. If it's not loaded, the TagsAndValues member
+// on the FAssetData contains the latest serialized data, which appears to get pruned once the UObject
+// is loaded.
+static bool ExtractTags(IAssetRegistry& InAssetRegistry, FAssetData& InAssetData, TMap<FName, FString>& OutTags)
+{
+	if (InAssetData.IsAssetLoaded())
+	{
+		// UObject::GetAssetRegistryTags might access render resources so we need to tag this thread
+		// as mentioned in HAL/ThreadingBase to avoid triggering an assert.
+		FTaskTagScope Scope(ETaskTag::EParallelRenderingThread);
+
+		UObject* Asset = InAssetData.GetAsset();
+
+		TArray<UObject::FAssetRegistryTag> AssetRegistryTags;
+		Asset->GetAssetRegistryTags(AssetRegistryTags);
+
+		for (TArray<UObject::FAssetRegistryTag>::TConstIterator TagIter(AssetRegistryTags); TagIter; ++TagIter)
+		{
+			OutTags.Add(TagIter->Name, TagIter->Value);
+		}
+	}
+	else
+	{
+		auto IsValid = [](const TTuple<FName, FString>& InTuple)
+		{
+			static FName NAME_FiBData("FiBData");
+			static FName NAME_ClassFlags("ClassFlags");
+			static FName NAME_AssetImportData("AssetImportData");
+
+			return InTuple.Key != NAME_FiBData
+				&& InTuple.Key != NAME_ClassFlags
+				&& InTuple.Key != NAME_AssetImportData;
+		};
+
+		Algo::TransformIf(InAssetData.TagsAndValues, OutTags, IsValid, [](const TTuple<FName, FString>& InTuple) { return InTuple; });
+	}
+
+	return true;
+}
+
+// Dependencies are straightforward. The AssetRegistry provides a nice interface for them.
+static bool ExtractDependencies(IAssetRegistry& InAssetRegistry, FAssetData& InAssetData, TArray<FName>& OutDependencies)
+{
+	return InAssetRegistry.GetDependencies(InAssetData.PackageName, OutDependencies);
+}
+
+// ThumbnailTools provides some nice features for us. The Thumbnail could be missing though.
+static bool ExtractThumbnail(IAssetRegistry& InAssetRegistry, FAssetData& InAssetData, FObjectThumbnail& OutThumbnail)
+{
+	FString ObjectFullName = InAssetData.GetFullName();
+
+	// Check if there's one cached already
+	const FObjectThumbnail* ExistingThumbnail = ThumbnailTools::FindCachedThumbnail(ObjectFullName);
+	if (ExistingThumbnail != nullptr)
+	{
+		OutThumbnail = *ExistingThumbnail;
+		return true;
+	}
+
+	// Load from disk instead
+	TArray< FName > ObjectFullNames;
+	FName ObjectFullNameFName(*ObjectFullName);
+	ObjectFullNames.Add(ObjectFullNameFName);
+
+	FThumbnailMap LoadedThumbnails;
+	if (ThumbnailTools::ConditionallyLoadThumbnailsForObjects(ObjectFullNames, LoadedThumbnails))
+	{
+		const FObjectThumbnail* LoadedThumbnail = LoadedThumbnails.Find(ObjectFullNameFName);
+		if (LoadedThumbnail != nullptr)
+		{
+			OutThumbnail = *LoadedThumbnail;
+			return true;
+		}
+	}
+
+	return false;
+}
+
+static bool WriteThumbnailToDisk(const FString& InThumbnailPath, const FObjectThumbnail& InObjectThumbnail, int32 InSize = 256)
+{
+	int32 Width = InObjectThumbnail.GetImageWidth();
+	int32 Height = InObjectThumbnail.GetImageHeight();
+	if (Width == 0 || Height == 0)
+	{
+#if UE_BUILD_DEBUG
+		UE_LOG(LogSourceControl, Log, TEXT("WriteThumbnailToDisk failed because dimensions are invalid for %s."), *InAssetPath);
+#endif
+		return false;
+	}
+
+	// Grab the - uncompressed - bytes
+	const TArray<uint8>& Bytes = InObjectThumbnail.GetUncompressedImageData();
+
+	// Convert the bytes to colors
+	TArray<FColor> Colors;
+	Colors.SetNumUninitialized(Bytes.Num() / sizeof(FColor));
+	FMemory::Memcpy(Colors.GetData(), Bytes.GetData(), Bytes.Num());
+
+	// Resize if needed
+	if (Width != InSize || Height != InSize)
+	{
+		TArray<FColor> ResizedColors;
+		ResizedColors.SetNum(InSize * InSize);
+
+		const bool bLinearSpace = false;
+		const bool bForceOpaqueOutput = false;
+		FImageUtils::ImageResize(Width, Height, Colors, InSize, InSize, ResizedColors, bLinearSpace, bForceOpaqueOutput);
+
+		Colors = MoveTemp(ResizedColors);
+	}
+
+	// Compress the image data
+	TArray64<uint8> CompressedBitmap;
+	FImageUtils::PNGCompressImageArray(InSize, InSize, TArrayView64<const FColor>(Colors.GetData(), Colors.Num()), CompressedBitmap);
+		
+	// Write to disk as a PNG to maintain transparency		
+	if (!FFileHelper::SaveArrayToFile(CompressedBitmap, *InThumbnailPath))
+	{
+#if UE_BUILD_DEBUG
+		UE_LOG(LogSourceControl, Log, TEXT("WriteThumbnailToDisk failed because image data could not be written to %s."), *InThumbnailPath);
+#endif
+		return false;
+	}
+
+	return true;
+}
+
+static bool WriteMetadataToDisk(const FString& InMetadataPath, const TMap<FName, FString>& InTags, const TArray<FName>& InDependencies)
+{
+	// Build JsonObject
+	TArray<TSharedPtr<FJsonValue>> JsonTagsArray;
+	JsonTagsArray.Reserve(InTags.Num());
+	for (const TPair<FName, FString> Tag : InTags)
+	{
+		TSharedRef<FJsonObject> JsonTag = MakeShareable(new FJsonObject);
+		JsonTag->SetStringField("name", Tag.Key.ToString());
+		JsonTag->SetStringField("value", Tag.Value);
+		JsonTagsArray.Add(MakeShareable(new FJsonValueObject(JsonTag)));
+	}
+
+	TArray<TSharedPtr<FJsonValue>> JsonDepsArray;
+	JsonDepsArray.Reserve(InDependencies.Num());
+	for (const FName& Dependency : InDependencies)
+	{
+		JsonDepsArray.Add(MakeShareable(new FJsonValueString(Dependency.ToString())));
+	}
+
+	TSharedRef<FJsonObject> JsonData = MakeShareable(new FJsonObject);
+	JsonData->SetArrayField("tags", JsonTagsArray);
+	JsonData->SetArrayField("deps", JsonDepsArray);
+
+	// Write Json to disk
+	FString JsonString;
+	TSharedRef<TJsonWriter<>> JsonWriter = TJsonWriterFactory<>::Create(&JsonString);
+	if (!FJsonSerializer::Serialize(JsonData, JsonWriter))
+	{
+#if UE_BUILD_DEBUG
+		UE_LOG(LogSourceControl, Log, TEXT("WriteMetadataToDisk failed because FJsonSerializer::Serialize failed."));
+#endif
+		return false;
+	}
+
+	if (!FFileHelper::SaveStringToFile(JsonString, *InMetadataPath))
+	{
+#if UE_BUILD_DEBUG
+		UE_LOG(LogSourceControl, Log, TEXT("WriteMetadataToDisk failed because SaveStringToFile %s failed."), *InAssetPath);
+#endif
+		return false;
+	}
+
+	return true;
+}
+
+bool ExtractMetadata(const FString& InPackagePath, const FString& InMetadataPath, const FString& InThumbnailPath, int InThumbnailSize /* = 256 */)
+{
+	FAssetRegistryModule& AssetRegistryModule = FModuleManager::LoadModuleChecked<FAssetRegistryModule>(TEXT("AssetRegistry"));
+	IAssetRegistry& AssetRegistry = AssetRegistryModule.Get();
+
+	FString PackageName;
+	if (!FPackageName::TryConvertFilenameToLongPackageName(InPackagePath, PackageName))
+	{
+#if UE_BUILD_DEBUG
+		UE_LOG(LogSourceControl, Log, TEXT("ExtractMetadata failed because TryConvertFilenameToLongPackageName %s failed."), *FileName);
+#endif
+		return false;
+	}
+
+	TArray<FAssetData> AssetDatas;
+	if (!AssetRegistry.GetAssetsByPackageName(*PackageName, AssetDatas, true) || AssetDatas.Num() == 0)
+	{
+#if UE_BUILD_DEBUG
+		UE_LOG(LogSourceControl, Log, TEXT("ExtractMetadata failed because GetAssetsByPackageName %s failed."), *FileName);
+#endif
+		return false;
+	}
+
+	FAssetData& AssetData = AssetDatas[0];
+
+	TMap<FName, FString> Tags;
+	ExtractTags(AssetRegistry, AssetData, Tags);
+
+	TArray<FName> Dependencies;
+	ExtractDependencies(AssetRegistry, AssetData, Dependencies);
+
+	FObjectThumbnail Thumbnail;
+	ExtractThumbnail(AssetRegistry, AssetData, Thumbnail);
+
+	bool bErrorWritingMetadata = false;
+	if (!WriteMetadataToDisk(InMetadataPath, Tags, Dependencies))
+	{
+		bErrorWritingMetadata = true;
+	}
+
+	bool bErrorWritingThumbnail = false;
+	if (!WriteThumbnailToDisk(InThumbnailPath, Thumbnail, InThumbnailSize))
+	{
+		bErrorWritingThumbnail = true;
+	}
+
+	return !bErrorWritingMetadata && !bErrorWritingThumbnail;
+}
+
+}
\ No newline at end of file
diff --git a/Engine/Plugins/Developer/SkeinSourceControl/Source/SkeinSourceControl/Private/SkeinSourceControlMetadata.h b/Engine/Plugins/Developer/SkeinSourceControl/Source/SkeinSourceControl/Private/SkeinSourceControlMetadata.h
new file mode 100644
--- /dev/null
+++ b/Engine/Plugins/Developer/SkeinSourceControl/Source/SkeinSourceControl/Private/SkeinSourceControlMetadata.h
@@ -0,0 +1,20 @@
+// Copyright Epic Games, Inc. All Rights Reserved.
+
+#pragma once
+
+#include "CoreMinimal.h"
+
+namespace SkeinSourceControlMetadata
+{
+
+/// <summary>
+/// Extracts the metadata Skein requires from the package to disk
+/// </summary>
+/// <param name="InPackagePath">Path to the package</param>
+/// <param name="InMetadataPath">Path where to store the metadata. If this is an empty string, no metadata will be extracted</param>
+/// <param name="InThumbnailPath">Path where to store the thumbnail. If this is an empty string, no thumbnail will be extracted</param>
+/// <param name="InThumbnailSize">Dimensions to use for the thumbnail.</param>
+/// <returns></returns>
+bool ExtractMetadata(const FString& InPackagePath, const FString& InMetadataPath, const FString& InThumbnailPath, int InThumbnailSize = 256);
+
+}
\ No newline at end of file
diff --git a/Engine/Plugins/Developer/SkeinSourceControl/Source/SkeinSourceControl/Private/SkeinSourceControlMetadataCommandlet.cpp b/Engine/Plugins/Developer/SkeinSourceControl/Source/SkeinSourceControl/Private/SkeinSourceControlMetadataCommandlet.cpp
new file mode 100644
--- /dev/null
+++ b/Engine/Plugins/Developer/SkeinSourceControl/Source/SkeinSourceControl/Private/SkeinSourceControlMetadataCommandlet.cpp
@@ -0,0 +1,89 @@
+// Copyright Epic Games, Inc. All Rights Reserved.
+
+#include "SkeinSourceControlMetadataCommandlet.h"
+#include "SkeinSourceControlMetadata.h"
+#include "SkeinSourceControlUtils.h"
+#include "Misc/Paths.h"
+
+DEFINE_LOG_CATEGORY_STATIC(LogSkeinMetadataCommandlet, Log, All);
+
+int32 USkeinSourceControlMetadataCommandlet::Main(FString const& Params)
+{
+	TArray<FString> Tokens;
+	TArray<FString> Switches;
+	ParseCommandLine(*Params, Tokens, Switches);
+
+	int32 Size = 256;
+	TArray<FString> Files;
+
+	for (const FString& Switch : Switches)
+	{
+		FString FilesSwitch;
+		FString FileCountsStr;
+		if (Switch.Split(TEXT("Files="), &FilesSwitch, &FileCountsStr) || Switch.Split(TEXT("f="), &FilesSwitch, &FileCountsStr))
+		{
+			TArray<FString> FileCountStrs;
+			if (FileCountsStr.ParseIntoArray(FileCountStrs, TEXT(",")) > 0)
+			{
+				Files = MoveTemp(FileCountStrs);
+			}
+		}
+		else if (FParse::Value(*Switch, TEXT("Size="), Size) || FParse::Value(*Switch, TEXT("s="), Size))
+		{
+			continue;
+		}
+	}
+
+	if (Files.IsEmpty())
+	{
+		UE_LOG(LogSkeinMetadataCommandlet, Warning, TEXT("You must specify one or more asset files using -Files=path/to/file1.uasset,path/to/file2.uasset"));
+		return -1;
+	}
+
+	int32 Failed = 0;
+	int32 Succeeded = 0;
+
+	for (const FString& FileName : Files)
+	{
+		if (FPaths::FileExists(FileName))
+		{
+			FString OutputFolder;
+			if (Tokens.Num() != 0)
+			{
+				OutputFolder = Tokens[0];
+			}
+			if (OutputFolder.IsEmpty())
+			{
+				OutputFolder = SkeinSourceControlUtils::FindSkeinIntermediateRoot(FileName);
+			}
+			if (OutputFolder.IsEmpty())
+			{
+				OutputFolder = FPaths::GetPath(FileName);
+			}
+			FPaths::ConvertRelativePathToFull(OutputFolder);
+			FPaths::NormalizeDirectoryName(OutputFolder);
+
+			FString MetadataFilename = SkeinSourceControlUtils::GetIntermediateMetadataPath(FileName, OutputFolder);
+			FString ThumbnailFilename = SkeinSourceControlUtils::GetIntermediateThumbnailPath(FileName, OutputFolder);
+
+			bool bExtractedMetadata = SkeinSourceControlMetadata::ExtractMetadata(FileName, MetadataFilename, ThumbnailFilename, Size);
+			if (bExtractedMetadata)
+			{
+				UE_LOG(LogSkeinMetadataCommandlet, Error, TEXT("Written thumbnail and metadata for %s"), *FileName);
+				++Succeeded;
+			}
+			else
+			{
+				UE_LOG(LogSkeinMetadataCommandlet, Error, TEXT("Could not write thumbnail or metadata for %s"), *FileName);
+				++Failed;
+			}
+		}
+		else
+		{
+			UE_LOG(LogSkeinMetadataCommandlet, Error, TEXT("Could not find input file: %s"), *FileName);
+			++Failed;
+		}
+	}
+
+	return Failed;
+}
diff --git a/Engine/Plugins/Developer/SkeinSourceControl/Source/SkeinSourceControl/Private/SkeinSourceControlMetadataCommandlet.h b/Engine/Plugins/Developer/SkeinSourceControl/Source/SkeinSourceControl/Private/SkeinSourceControlMetadataCommandlet.h
new file mode 100644
--- /dev/null
+++ b/Engine/Plugins/Developer/SkeinSourceControl/Source/SkeinSourceControl/Private/SkeinSourceControlMetadataCommandlet.h
@@ -0,0 +1,19 @@
+// Copyright Epic Games, Inc. All Rights Reserved.
+
+#pragma once
+
+#include "CoreMinimal.h"
+#include "UObject/ObjectMacros.h"
+#include "Commandlets/Commandlet.h"
+#include "SkeinSourceControlMetadataCommandlet.generated.h"
+
+UCLASS()
+class USkeinSourceControlMetadataCommandlet : public UCommandlet
+{
+	GENERATED_BODY()
+
+public:
+	//~ Begin UCommandlet Interface
+	virtual int32 Main(FString const& Params) override;
+	//~ End UCommandlet Interface
+};
\ No newline at end of file
diff --git a/Engine/Plugins/Developer/SkeinSourceControl/Source/SkeinSourceControl/Private/SkeinSourceControlOperations.cpp b/Engine/Plugins/Developer/SkeinSourceControl/Source/SkeinSourceControl/Private/SkeinSourceControlOperations.cpp
--- a/Engine/Plugins/Developer/SkeinSourceControl/Source/SkeinSourceControl/Private/SkeinSourceControlOperations.cpp
+++ b/Engine/Plugins/Developer/SkeinSourceControl/Source/SkeinSourceControl/Private/SkeinSourceControlOperations.cpp
@@ -1,7 +1,7 @@
 // Copyright Epic Games, Inc. All Rights Reserved.
 
 #include "SkeinSourceControlOperations.h"
-#include "SkeinSourceControlThumbnail.h"
+#include "SkeinSourceControlMetadata.h"
 #include "SkeinSourceControlProvider.h"
 #include "SkeinSourceControlCommand.h"
 #include "SkeinSourceControlModule.h"
@@ -82,9 +82,40 @@ bool FSkeinCheckInWorker::Execute(FSkeinSourceControlCommand& InCommand)
 {
 	check(InCommand.Operation->GetName() == GetName());
 
-	InCommand.bCommandSuccessful = SkeinSourceControlUtils::RunCommand(TEXT("projects commit"), InCommand.SkeinBinaryPath, InCommand.SkeinProjectRoot, TArray<FString>(), InCommand.Files, InCommand.InfoMessages, InCommand.ErrorMessages);
+	TSharedRef<FCheckIn, ESPMode::ThreadSafe> Operation = StaticCastSharedRef<FCheckIn>(InCommand.Operation);
 
-	SkeinSourceControlUtils::RunUpdateStatus(InCommand.SkeinBinaryPath, InCommand.SkeinProjectRoot, InCommand.Files, InCommand.ErrorMessages, States);
+	// Gather all files under our control - they will all end up in the snapshot
+	TArray<FString> ProjectRoots;
+	ProjectRoots.Add(InCommand.SkeinProjectRoot);
+
+	SkeinSourceControlUtils::RunUpdateStatus(InCommand.SkeinBinaryPath, InCommand.SkeinProjectRoot, ProjectRoots, InCommand.ErrorMessages, States);
+
+	// Export properties for each of them (metadata + thumbnail)
+	ParallelFor(States.Num(),
+		[&] (int32 ElemIdx)
+		{
+			FSkeinSourceControlState& FileState = States[ElemIdx];
+			if (FileState.IsSourceControlled())
+			{
+				FString& File = FileState.Filename;
+
+				const FString FileMetadata = SkeinSourceControlUtils::GetIntermediateMetadataPath(File);
+				const FString FileThumbnail = SkeinSourceControlUtils::GetIntermediateThumbnailPath(File);
+
+				SkeinSourceControlMetadata::ExtractMetadata(File, FileMetadata, FileThumbnail);
+			}
+		}
+	);
+
+	// Grab the message for this snapshot
+	TArray<FString> Parameters;
+	Parameters.Add("--message=\"" + Operation->GetDescription().ToString() + "\"");
+
+	// Create the snapshot
+	InCommand.bCommandSuccessful = SkeinSourceControlUtils::RunCommand(TEXT("projects snapshots create"), InCommand.SkeinBinaryPath, InCommand.SkeinProjectRoot, Parameters, TArray<FString>() /* InCommand.Files */, InCommand.InfoMessages, InCommand.ErrorMessages);
+
+	// Cache the new file states
+	SkeinSourceControlUtils::RunUpdateStatus(InCommand.SkeinBinaryPath, InCommand.SkeinProjectRoot, ProjectRoots, InCommand.ErrorMessages, States);
 
 	return InCommand.bCommandSuccessful;
 }
@@ -180,7 +211,7 @@ bool FSkeinSyncWorker::Execute(FSkeinSourceControlCommand& InCommand)
 {
 	check(InCommand.Operation->GetName() == GetName());
 
-	InCommand.bCommandSuccessful = SkeinSourceControlUtils::RunCommand(TEXT("projects pull"), InCommand.SkeinBinaryPath, InCommand.SkeinProjectRoot, TArray<FString>(), TArray<FString>(), InCommand.InfoMessages, InCommand.ErrorMessages);
+	InCommand.bCommandSuccessful = SkeinSourceControlUtils::RunCommand(TEXT("projects snapshots get"), InCommand.SkeinBinaryPath, InCommand.SkeinProjectRoot, TArray<FString>(), TArray<FString>(), InCommand.InfoMessages, InCommand.ErrorMessages);
 
 	SkeinSourceControlUtils::RunUpdateStatus(InCommand.SkeinBinaryPath, InCommand.SkeinProjectRoot, InCommand.Files, InCommand.ErrorMessages, States);
 
diff --git a/Engine/Plugins/Developer/SkeinSourceControl/Source/SkeinSourceControl/Private/SkeinSourceControlProvider.cpp b/Engine/Plugins/Developer/SkeinSourceControl/Source/SkeinSourceControl/Private/SkeinSourceControlProvider.cpp
--- a/Engine/Plugins/Developer/SkeinSourceControl/Source/SkeinSourceControl/Private/SkeinSourceControlProvider.cpp
+++ b/Engine/Plugins/Developer/SkeinSourceControl/Source/SkeinSourceControl/Private/SkeinSourceControlProvider.cpp
@@ -1,7 +1,6 @@
 // Copyright Epic Games, Inc. All Rights Reserved.
 
 #include "SkeinSourceControlProvider.h"
-#include "SkeinSourceControlThumbnail.h"
 #include "SkeinSourceControlCommand.h"
 #include "SkeinSourceControlState.h"
 #include "SkeinSourceControlUtils.h"
diff --git a/Engine/Plugins/Developer/SkeinSourceControl/Source/SkeinSourceControl/Private/SkeinSourceControlThumbnail.cpp b/Engine/Plugins/Developer/SkeinSourceControl/Source/SkeinSourceControl/Private/SkeinSourceControlThumbnail.cpp
deleted file mode 100644
diff --git a/Engine/Plugins/Developer/SkeinSourceControl/Source/SkeinSourceControl/Private/SkeinSourceControlThumbnail.h b/Engine/Plugins/Developer/SkeinSourceControl/Source/SkeinSourceControl/Private/SkeinSourceControlThumbnail.h
deleted file mode 100644
diff --git a/Engine/Plugins/Developer/SkeinSourceControl/Source/SkeinSourceControl/Private/SkeinSourceControlThumbnailCommandlet.cpp b/Engine/Plugins/Developer/SkeinSourceControl/Source/SkeinSourceControl/Private/SkeinSourceControlThumbnailCommandlet.cpp
deleted file mode 100644
diff --git a/Engine/Plugins/Developer/SkeinSourceControl/Source/SkeinSourceControl/Private/SkeinSourceControlThumbnailCommandlet.h b/Engine/Plugins/Developer/SkeinSourceControl/Source/SkeinSourceControl/Private/SkeinSourceControlThumbnailCommandlet.h
deleted file mode 100644
diff --git a/Engine/Plugins/Developer/SkeinSourceControl/Source/SkeinSourceControl/Private/SkeinSourceControlUtils.cpp b/Engine/Plugins/Developer/SkeinSourceControl/Source/SkeinSourceControl/Private/SkeinSourceControlUtils.cpp
--- a/Engine/Plugins/Developer/SkeinSourceControl/Source/SkeinSourceControl/Private/SkeinSourceControlUtils.cpp
+++ b/Engine/Plugins/Developer/SkeinSourceControl/Source/SkeinSourceControl/Private/SkeinSourceControlUtils.cpp
@@ -6,10 +6,10 @@
 #include "HAL/FileManager.h"
 #include "Misc/Paths.h"
 #include "Misc/ScopeLock.h"
+#include "Misc/FileHelper.h"
 #include "ISourceControlModule.h"
 #include "Serialization/JsonSerializer.h"
 #include "Dom/JsonValue.h"
-#include "Misc/FileHelper.h"
 
 namespace SkeinSourceControlConstants
 {
@@ -44,41 +44,73 @@ static void ParseStatusOutput(const FString& InSkeinProjectRoot, const TArray<FS
 		}
 	}
 
-	// Iterate on all files explicitly listed in the command
+	// Separate InFiles between Files and Directories
+	// Files require an exact match, Directories only a StartsWith match
+	TSet<FString> Files;
+	TArray<FString> Directories;
 	for (const auto& File : InFiles)
 	{
+		const bool bIsDirectory = IFileManager::Get().DirectoryExists(*File);
+		if (bIsDirectory)
+		{
+			Directories.Add(File);
+		}
+		else
+		{
+			Files.Add(File);
+		}
+	}
+
+	// Inspect the Results and return what's needed
+	for (const auto& Result : FileStates)
+	{
+		const FString& File = Result.Key;
+
 		FSkeinSourceControlState FileState(File);
 		FileState.TimeStamp = Now;
 		FileState.State = ESkeinState::Unknown;
 
-		FString* State = FileStates.Find(File);
-		if (State != nullptr)
+		bool bReturnResult = (Files.Find(File) != nullptr);
+		if (!bReturnResult)
+		{
+			for (const FString& Directory : Directories)
+			{
+				if (File.StartsWith(Directory))
+				{
+					bReturnResult = true;
+					break;
+				}
+			}
+		}
+
+		if (bReturnResult)
 		{
-			if (*State == "unknown")
+			FString State = Result.Value;
+			if (State == "unknown")
 			{
 				FileState.State = ESkeinState::Unknown;
 			}
-			else if (*State == "add")
+			else if (State == "add")
 			{
 				FileState.State = ESkeinState::Added;
 			}
-			else if (*State == "remove")
+			else if (State == "remove")
 			{
 				FileState.State = ESkeinState::Deleted;
 			}
-			else if (*State == "missing")
+			else if (State == "missing")
 			{
 				FileState.State = ESkeinState::Missing;
 			}
-			else if (*State == "modified")
+			else if (State == "modified")
 			{
 				FileState.State = ESkeinState::Modified;
 			}
-			else if (*State == "untracked")
+			else if (State == "untracked")
 			{
 				FileState.State = ESkeinState::NotControlled;
 			}
-			else if (*State == "unchanged")
+			else if (State == "unchanged")
 			{
 				FileState.State = ESkeinState::Unchanged;
 			}
@@ -86,9 +118,9 @@ static void ParseStatusOutput(const FString& InSkeinProjectRoot, const TArray<FS
 			{
 				checkNoEntry();
 			}
-		}
 
-		OutStates.Add(FileState);
+			OutStates.Add(FileState);
+		}
 	}
 }
 
@@ -113,6 +145,10 @@ static bool RunCommandInternalRaw(const FString& InCommand, const FString& InSke
 	LoggableCommand += TEXT(" ");
 	LoggableCommand += TEXT("--json");
 
+	// Append the "--quiet" param to indicate we don't want any Yes/No prompts
+	LoggableCommand += TEXT(" ");
+	LoggableCommand += TEXT("--quiet");
+
 	// Append to the command all parameters, and then finally the files
 	for(const auto& Parameter : InParameters)
 	{
@@ -248,6 +284,64 @@ FString FindSkeinProjectRoot(const FString& InPath)
 	return PathToSkeinProjectRoot;
 }
 
+/// <notes>
+/// This returns the intermediate root folder where SkeinCLI expects to find temporary metadata.
+/// </notes>
+FString FindSkeinIntermediateRoot(const FString& InPath)
+{
+	FString PathToSkeinProjectRoot = FindSkeinProjectRoot(InPath);
+	if (!PathToSkeinProjectRoot.IsEmpty())
+	{
+		FTCHARToUTF8 SkeinProjectRootLowerCase(*PathToSkeinProjectRoot.ToLower());
+		FString ProjectRootHash = FMD5::HashBytes((const uint8*)SkeinProjectRootLowerCase.Get(), SkeinProjectRootLowerCase.Length());
+		FString PlatformTempDir = FPlatformProcess::UserTempDir();
+
+		return FPaths::Combine(PlatformTempDir, ProjectRootHash);
+	}
+	return FString();
+}
+
+FString GetIntermediatePath(const FString& InPath, const FString& InIntermediateRoot)
+{
+	FString PathToSkeinProjectRoot = FindSkeinProjectRoot(InPath);
+	FString PathToSkeinIntermediateRoot = InIntermediateRoot.IsEmpty() ? FindSkeinIntermediateRoot(InPath) : InIntermediateRoot;
+	if (!PathToSkeinProjectRoot.IsEmpty() && !PathToSkeinIntermediateRoot.IsEmpty())
+	{
+		PathToSkeinProjectRoot += FPlatformMisc::GetDefaultPathSeparator();
+
+		FString PathRelative = *InPath;
+		if (FPaths::MakePathRelativeTo(PathRelative, *PathToSkeinProjectRoot))
+		{
+			FString Result = FPaths::Combine(PathToSkeinIntermediateRoot, PathRelative);
+			FPaths::NormalizeFilename(Result);
+
+			return Result;
+		}
+	}
+
+	return FString();
+}
+
+FString GetIntermediateMetadataPath(const FString& InPath, const FString& InIntermediateRoot)
+{
+	FString IntermediatePath = GetIntermediatePath(InPath, InIntermediateRoot);
+	if (!IntermediatePath.IsEmpty())
+	{
+		return FPaths::GetBaseFilename(IntermediatePath, false) + ".metadata.json";
+	}
+	return FString();
+}
+
+FString GetIntermediateThumbnailPath(const FString& InPath, const FString& InIntermediateRoot)
+{
+	FString IntermediatePath = GetIntermediatePath(InPath, InIntermediateRoot);
+	if (!IntermediatePath.IsEmpty())
+	{
+		return FPaths::GetBaseFilename(IntermediatePath, false) + ".thumbnail.png";
+	}
+	return FString();
+}
+
 bool IsSkeinAvailable()
 {
 	FString SkeinBinaryPath = FindSkeinBinaryPath();
@@ -370,8 +464,15 @@ bool RunUpdateStatus(const FString& InSkeinBinaryPath, const FString& InSkeinPro
 	TSet<FString> UniquePaths;
 	for (const auto& File : InFiles)
 	{
-		const FString Path = FPaths::GetPath(*File);
-		UniquePaths.Add(Path);
+		const bool bIsDirectory = IFileManager::Get().DirectoryExists(*File);
+		if (bIsDirectory)
+		{
+			UniquePaths.Add(*File);
+		}
+		else
+		{
+			UniquePaths.Add(FPaths::GetPath(*File));
+		}
 	}
 
 	int NumErrors = 0;
diff --git a/Engine/Plugins/Developer/SkeinSourceControl/Source/SkeinSourceControl/Private/SkeinSourceControlUtils.h b/Engine/Plugins/Developer/SkeinSourceControl/Source/SkeinSourceControl/Private/SkeinSourceControlUtils.h
--- a/Engine/Plugins/Developer/SkeinSourceControl/Source/SkeinSourceControl/Private/SkeinSourceControlUtils.h
+++ b/Engine/Plugins/Developer/SkeinSourceControl/Source/SkeinSourceControl/Private/SkeinSourceControlUtils.h
@@ -20,6 +20,26 @@ FString FindSkeinBinaryPath();
  */
 FString FindSkeinProjectRoot(const FString& InPath);
 
+/**
+ * Determines the path to the Skein intermediate root (if any)
+ * @returns an empty string if the directory is not part of a Skein project
+ */
+FString FindSkeinIntermediateRoot(const FString& InPath);
+
+/**
+  * Returns the path to store the extracted metadata of the given file.
+  * This should be in sync with what the CLI expects.
+  * @returns an empty string if the directory is not part of a Skein project
+  */
+FString GetIntermediateMetadataPath(const FString& InPath, const FString& InIntermediateRoot = FString());
+
+/**
+  * Returns the path to store the extracted thumbnail of the given file.
+  * This should be in sync with what the CLI expects.
+  * @returns an empty string if the directory is not part of a Skein project
+  */
+FString GetIntermediateThumbnailPath(const FString& InPath, const FString& InIntermediateRoot = FString());
+
 /**
  * Finds out if the Skein environment is available.
  * In practice this checks if the Skein CLI application is installed at the expected location.
