SkeinUE: Compile fix for NonUnity builds.

diff --git a/Engine/Plugins/Developer/SkeinSourceControl/Source/SkeinSourceControl/Private/SkeinSourceControlMetadata.cpp b/Engine/Plugins/Developer/SkeinSourceControl/Source/SkeinSourceControl/Private/SkeinSourceControlMetadata.cpp
--- a/Engine/Plugins/Developer/SkeinSourceControl/Source/SkeinSourceControl/Private/SkeinSourceControlMetadata.cpp
+++ b/Engine/Plugins/Developer/SkeinSourceControl/Source/SkeinSourceControl/Private/SkeinSourceControlMetadata.cpp
@@ -6,6 +6,7 @@
 #include "Serialization/JsonSerializer.h"
 #include "Misc/AssetRegistryInterface.h"
 #include "Misc/FileHelper.h"
+#include "Algo/Transform.h"
 #include "ISourceControlModule.h"
 #include "ObjectTools.h"
 #include "ImageUtils.h"
@@ -37,6 +38,7 @@ static bool ExtractTags(IAssetRegistry& InAssetRegistry, FAssetData& InAssetData
 	}
 	else
 	{
+
 		auto IsValid = [](const TTuple<FName, FString>& InTuple)
 		{
 			static FName NAME_FiBData("FiBData");
@@ -48,7 +50,7 @@ static bool ExtractTags(IAssetRegistry& InAssetRegistry, FAssetData& InAssetData
 				&& InTuple.Key != NAME_AssetImportData;
 		};
 
-		Algo::TransformIf(InAssetData.TagsAndValues, OutTags, IsValid, [](const TTuple<FName, FString>& InTuple) { return InTuple; });
+		Algo::TransformIf(InAssetData.TagsAndValues.CopyMap(), OutTags, IsValid, [](const TTuple<FName, FString>& InTuple) { return InTuple; });
 	}
 
 	return true;
diff --git a/Engine/Plugins/Developer/SkeinSourceControl/Source/SkeinSourceControl/Private/SkeinSourceControlOperations.cpp b/Engine/Plugins/Developer/SkeinSourceControl/Source/SkeinSourceControl/Private/SkeinSourceControlOperations.cpp
--- a/Engine/Plugins/Developer/SkeinSourceControl/Source/SkeinSourceControl/Private/SkeinSourceControlOperations.cpp
+++ b/Engine/Plugins/Developer/SkeinSourceControl/Source/SkeinSourceControl/Private/SkeinSourceControlOperations.cpp
@@ -8,6 +8,7 @@
 #include "SkeinSourceControlUtils.h"
 #include "SkeinSourceControlState.h"
 #include "Modules/ModuleManager.h"
+#include "Async/ParallelFor.h"
 #include "SourceControlOperations.h"
 
 #define LOCTEXT_NAMESPACE "SkeinSourceControl"
diff --git a/Engine/Plugins/Developer/SkeinSourceControl/Source/SkeinSourceControl/Private/SkeinSourceControlUtils.cpp b/Engine/Plugins/Developer/SkeinSourceControl/Source/SkeinSourceControl/Private/SkeinSourceControlUtils.cpp
--- a/Engine/Plugins/Developer/SkeinSourceControl/Source/SkeinSourceControl/Private/SkeinSourceControlUtils.cpp
+++ b/Engine/Plugins/Developer/SkeinSourceControl/Source/SkeinSourceControl/Private/SkeinSourceControlUtils.cpp
@@ -4,9 +4,10 @@
 #include "HAL/PlatformProcess.h"
 #include "HAL/PlatformFileManager.h"
 #include "HAL/FileManager.h"
+#include "Misc/FileHelper.h"
 #include "Misc/Paths.h"
+#include "Misc/SecureHash.h"
 #include "Misc/ScopeLock.h"
-#include "Misc/FileHelper.h"
 #include "ISourceControlModule.h"
 #include "Serialization/JsonSerializer.h"
 #include "Dom/JsonValue.h"
