SkeinUE: Compile fix for Linux build.

diff --git a/Engine/Plugins/Developer/SkeinSourceControl/Source/SkeinSourceControl/Private/SkeinSourceControlMetadata.cpp b/Engine/Plugins/Developer/SkeinSourceControl/Source/SkeinSourceControl/Private/SkeinSourceControlMetadata.cpp
--- a/Engine/Plugins/Developer/SkeinSourceControl/Source/SkeinSourceControl/Private/SkeinSourceControlMetadata.cpp
+++ b/Engine/Plugins/Developer/SkeinSourceControl/Source/SkeinSourceControl/Private/SkeinSourceControlMetadata.cpp
@@ -148,7 +148,7 @@ static bool WriteMetadataToDisk(const FString& InMetadataPath, const TMap<FName,
 	// Build JsonObject
 	TArray<TSharedPtr<FJsonValue>> JsonTagsArray;
 	JsonTagsArray.Reserve(InTags.Num());
-	for (const TPair<FName, FString> Tag : InTags)
+	for (const TPair<FName, FString>& Tag : InTags)
 	{
 		TSharedRef<FJsonObject> JsonTag = MakeShareable(new FJsonObject);
 		JsonTag->SetStringField("name", Tag.Key.ToString());
