SkeinUE: Compile fix for Mac build.

diff --git a/Engine/Plugins/Developer/SkeinSourceControl/Source/SkeinSourceControl/Private/SkeinSourceControlMetadata.cpp b/Engine/Plugins/Developer/SkeinSourceControl/Source/SkeinSourceControl/Private/SkeinSourceControlMetadata.cpp
--- a/Engine/Plugins/Developer/SkeinSourceControl/Source/SkeinSourceControl/Private/SkeinSourceControlMetadata.cpp
+++ b/Engine/Plugins/Developer/SkeinSourceControl/Source/SkeinSourceControl/Private/SkeinSourceControlMetadata.cpp
@@ -101,7 +101,7 @@ static bool WriteThumbnailToDisk(const FString& InThumbnailPath, const FObjectTh
 	if (Width == 0 || Height == 0)
 	{
 #if UE_BUILD_DEBUG
-		UE_LOG(LogSourceControl, Log, TEXT("WriteThumbnailToDisk failed because dimensions are invalid for %s."), *InAssetPath);
+		UE_LOG(LogSourceControl, Log, TEXT("WriteThumbnailToDisk failed because dimensions are invalid for %s."), *InThumbnailPath);
 #endif
 		return false;
 	}
@@ -181,7 +181,7 @@ static bool WriteMetadataToDisk(const FString& InMetadataPath, const TMap<FName,
 	if (!FFileHelper::SaveStringToFile(JsonString, *InMetadataPath))
 	{
 #if UE_BUILD_DEBUG
-		UE_LOG(LogSourceControl, Log, TEXT("WriteMetadataToDisk failed because SaveStringToFile %s failed."), *InAssetPath);
+		UE_LOG(LogSourceControl, Log, TEXT("WriteMetadataToDisk failed because SaveStringToFile %s failed."), *InMetadataPath);
 #endif
 		return false;
 	}
@@ -198,7 +198,7 @@ bool ExtractMetadata(const FString& InPackagePath, const FString& InMetadataPath
 	if (!FPackageName::TryConvertFilenameToLongPackageName(InPackagePath, PackageName))
 	{
 #if UE_BUILD_DEBUG
-		UE_LOG(LogSourceControl, Log, TEXT("ExtractMetadata failed because TryConvertFilenameToLongPackageName %s failed."), *FileName);
+		UE_LOG(LogSourceControl, Log, TEXT("ExtractMetadata failed because TryConvertFilenameToLongPackageName %s failed."), *InPackagePath);
 #endif
 		return false;
 	}
@@ -207,7 +207,7 @@ bool ExtractMetadata(const FString& InPackagePath, const FString& InMetadataPath
 	if (!AssetRegistry.GetAssetsByPackageName(*PackageName, AssetDatas, true) || AssetDatas.Num() == 0)
 	{
 #if UE_BUILD_DEBUG
-		UE_LOG(LogSourceControl, Log, TEXT("ExtractMetadata failed because GetAssetsByPackageName %s failed."), *FileName);
+		UE_LOG(LogSourceControl, Log, TEXT("ExtractMetadata failed because GetAssetsByPackageName %s failed."), *InPackagePath);
 #endif
 		return false;
 	}
