SkeinUE: Bump CLI Batch size to a still conservative 32.

diff --git a/Engine/Plugins/Developer/SkeinSourceControl/Source/SkeinSourceControl/Private/SkeinSourceControlUtils.cpp b/Engine/Plugins/Developer/SkeinSourceControl/Source/SkeinSourceControl/Private/SkeinSourceControlUtils.cpp
--- a/Engine/Plugins/Developer/SkeinSourceControl/Source/SkeinSourceControl/Private/SkeinSourceControlUtils.cpp
+++ b/Engine/Plugins/Developer/SkeinSourceControl/Source/SkeinSourceControl/Private/SkeinSourceControlUtils.cpp
@@ -15,7 +15,7 @@
 namespace SkeinSourceControlConstants
 {
 	/** The maximum number of files we submit in a single Skein command */
-	const int32 MaxFilesPerBatch = 1;
+	const int32 MaxFilesPerBatch = 32;
 }
 
 namespace SkeinSourceControlUtils
