SkeinUE: Settings dialog should monitor for new projects as we initialize them outside of the UE environment using CLI.

diff --git a/Engine/Plugins/Developer/SkeinSourceControl/Source/SkeinSourceControl/Private/SSkeinSourceControlSettings.cpp b/Engine/Plugins/Developer/SkeinSourceControl/Source/SkeinSourceControl/Private/SSkeinSourceControlSettings.cpp
--- a/Engine/Plugins/Developer/SkeinSourceControl/Source/SkeinSourceControl/Private/SSkeinSourceControlSettings.cpp
+++ b/Engine/Plugins/Developer/SkeinSourceControl/Source/SkeinSourceControl/Private/SSkeinSourceControlSettings.cpp
@@ -26,6 +26,7 @@ void SSkeinSourceControlSettings::Construct(const FArguments& InArgs)
 			[
 				SNew(STextBlock)
 				.Text(LOCTEXT("SkeinCLINotFound", "The Skein Command Line application could not be found."))
+				.ColorAndOpacity(FLinearColor::Red)
 				.WrapTextAt(450.0f)
 				.ToolTipText(LOCTEXT("SkeinCLINotFound_Tooltip", "Without the Skein Command Line application the Unreal Editor cannot communicate to the Skein cloud server. Please make sure it's installed at the correct location."))
 			]
@@ -43,6 +44,7 @@ void SSkeinSourceControlSettings::Construct(const FArguments& InArgs)
 			[
 				SNew(STextBlock)
 				.Text(LOCTEXT("SkeinProjectNotFound", "There is no Skein project initialized for this location."))
+				.ColorAndOpacity(FLinearColor::Red)
 				.WrapTextAt(450.0f)
 				.ToolTipText(LOCTEXT("SkeinProjectNotFound_Tooltip", "Skein projects are created in the Skein Web UI and then cloned to your local machine using the Skein Command Line application. Once complete, you can use the Skein Source Control plugin to manage the project assets from within Unreal Editor."))
 			]
diff --git a/Engine/Plugins/Developer/SkeinSourceControl/Source/SkeinSourceControl/Private/SkeinSourceControlProvider.cpp b/Engine/Plugins/Developer/SkeinSourceControl/Source/SkeinSourceControl/Private/SkeinSourceControlProvider.cpp
--- a/Engine/Plugins/Developer/SkeinSourceControl/Source/SkeinSourceControl/Private/SkeinSourceControlProvider.cpp
+++ b/Engine/Plugins/Developer/SkeinSourceControl/Source/SkeinSourceControl/Private/SkeinSourceControlProvider.cpp
@@ -45,11 +45,26 @@ FText FSkeinSourceControlProvider::GetStatusText() const
 
 bool FSkeinSourceControlProvider::IsAvailable() const
 {
+	if (!bSkeinBinaryFound)
+	{
+		bSkeinBinaryFound = SkeinSourceControlUtils::IsSkeinAvailable();
+	}
+
 	return bSkeinBinaryFound;
 }
 
 bool FSkeinSourceControlProvider::IsEnabled() const
 {
+	if (!bSkeinBinaryFound)
+	{
+		bSkeinBinaryFound = SkeinSourceControlUtils::IsSkeinAvailable();
+	}
+
+	if (!bSkeinProjectFound)
+	{
+		bSkeinProjectFound = SkeinSourceControlUtils::IsSkeinProjectFound(FPaths::ProjectDir(), ProjectRoot, ProjectName);
+	}
+
 	return bSkeinBinaryFound && bSkeinProjectFound;
 }
 
diff --git a/Engine/Plugins/Developer/SkeinSourceControl/Source/SkeinSourceControl/Private/SkeinSourceControlProvider.h b/Engine/Plugins/Developer/SkeinSourceControl/Source/SkeinSourceControl/Private/SkeinSourceControlProvider.h
--- a/Engine/Plugins/Developer/SkeinSourceControl/Source/SkeinSourceControl/Private/SkeinSourceControlProvider.h
+++ b/Engine/Plugins/Developer/SkeinSourceControl/Source/SkeinSourceControl/Private/SkeinSourceControlProvider.h
@@ -94,19 +94,19 @@ private:
 private:
 
 	/** Is Skein service available */
-	bool bSkeinBinaryFound;
+	mutable bool bSkeinBinaryFound;
 
 	/** Is Skein project found */
-	bool bSkeinProjectFound;
-
-	/** Skein binary path */
-	FString BinaryPath;
+	mutable bool bSkeinProjectFound;
 
 	/** Skein project root */
-	FString ProjectRoot;
+	mutable FString ProjectRoot;
 
 	/** Skein project name */
-	FString ProjectName;
+	mutable FString ProjectName;
+
+	/** Skein binary path */
+	FString BinaryPath;
 
 	/** State cache */
 	TMap<FString, TSharedRef<FSkeinSourceControlState, ESPMode::ThreadSafe>> StateCache;
