SkeinUE: The 'CheckIn' operation should generate metadata for 'All' files in the snapshot and update the state for each of them.

diff --git a/Engine/Plugins/Developer/SkeinSourceControl/Source/SkeinSourceControl/Private/SkeinSourceControlOperations.cpp b/Engine/Plugins/Developer/SkeinSourceControl/Source/SkeinSourceControl/Private/SkeinSourceControlOperations.cpp
--- a/Engine/Plugins/Developer/SkeinSourceControl/Source/SkeinSourceControl/Private/SkeinSourceControlOperations.cpp
+++ b/Engine/Plugins/Developer/SkeinSourceControl/Source/SkeinSourceControl/Private/SkeinSourceControlOperations.cpp
@@ -89,7 +89,10 @@ bool FSkeinCheckInWorker::Execute(FSkeinSourceControlCommand& InCommand)
 	TArray<FString> ProjectRoots;
 	ProjectRoots.Add(InCommand.SkeinProjectRoot);
 
-	SkeinSourceControlUtils::RunUpdateStatus(InCommand.SkeinBinaryPath, InCommand.SkeinProjectRoot, ProjectRoots, InCommand.ErrorMessages, States);
+	TArray<FString> StatusParameters;
+	StatusParameters.Add("--all");
+
+	SkeinSourceControlUtils::RunUpdateStatus(InCommand.SkeinBinaryPath, InCommand.SkeinProjectRoot, StatusParameters, ProjectRoots, InCommand.ErrorMessages, States);
 
 	// Export properties for each of them (metadata + thumbnail)
 	ParallelFor(States.Num(),
@@ -116,7 +119,7 @@ bool FSkeinCheckInWorker::Execute(FSkeinSourceControlCommand& InCommand)
 	InCommand.bCommandSuccessful = SkeinSourceControlUtils::RunCommand(TEXT("projects snapshots create"), InCommand.SkeinBinaryPath, InCommand.SkeinProjectRoot, Parameters, TArray<FString>() /* InCommand.Files */, InCommand.InfoMessages, InCommand.ErrorMessages);
 
 	// Cache the new file states
-	SkeinSourceControlUtils::RunUpdateStatus(InCommand.SkeinBinaryPath, InCommand.SkeinProjectRoot, ProjectRoots, InCommand.ErrorMessages, States);
+	SkeinSourceControlUtils::RunUpdateStatus(InCommand.SkeinBinaryPath, InCommand.SkeinProjectRoot, StatusParameters, ProjectRoots, InCommand.ErrorMessages, States);
 
 	return InCommand.bCommandSuccessful;
 }
@@ -141,7 +144,7 @@ bool FSkeinMarkForAddWorker::Execute(FSkeinSourceControlCommand& InCommand)
 
 	InCommand.bCommandSuccessful = SkeinSourceControlUtils::RunCommand(TEXT("assets track"), InCommand.SkeinBinaryPath, InCommand.SkeinProjectRoot, TArray<FString>(), InCommand.Files, InCommand.InfoMessages, InCommand.ErrorMessages);
 
-	SkeinSourceControlUtils::RunUpdateStatus(InCommand.SkeinBinaryPath, InCommand.SkeinProjectRoot, InCommand.Files, InCommand.ErrorMessages, States);
+	SkeinSourceControlUtils::RunUpdateStatus(InCommand.SkeinBinaryPath, InCommand.SkeinProjectRoot, TArray<FString>(), InCommand.Files, InCommand.ErrorMessages, States);
 
 	return InCommand.bCommandSuccessful;
 }
@@ -164,9 +167,9 @@ bool FSkeinDeleteWorker::Execute(FSkeinSourceControlCommand& InCommand)
 {
 	check(InCommand.Operation->GetName() == GetName());
 
-	InCommand.bCommandSuccessful = SkeinSourceControlUtils::RunCommand(TEXT("assets untrack"), InCommand.SkeinBinaryPath, InCommand.SkeinProjectRoot, TArray<FString>(), InCommand.Files, InCommand.InfoMessages, InCommand.ErrorMessages);
+	InCommand.bCommandSuccessful = SkeinSourceControlUtils::RunCommand(TEXT("assets delete"), InCommand.SkeinBinaryPath, InCommand.SkeinProjectRoot, TArray<FString>(), InCommand.Files, InCommand.InfoMessages, InCommand.ErrorMessages);
 
-	SkeinSourceControlUtils::RunUpdateStatus(InCommand.SkeinBinaryPath, InCommand.SkeinProjectRoot, InCommand.Files, InCommand.ErrorMessages, States);
+	SkeinSourceControlUtils::RunUpdateStatus(InCommand.SkeinBinaryPath, InCommand.SkeinProjectRoot, TArray<FString>(), InCommand.Files, InCommand.ErrorMessages, States);
 
 	return InCommand.bCommandSuccessful;
 }
@@ -191,7 +194,7 @@ bool FSkeinRevertWorker::Execute(FSkeinSourceControlCommand& InCommand)
 
 	InCommand.bCommandSuccessful = SkeinSourceControlUtils::RunCommand(TEXT("assets revert"), InCommand.SkeinBinaryPath, InCommand.SkeinProjectRoot, TArray<FString>(), InCommand.Files, InCommand.InfoMessages, InCommand.ErrorMessages);
 
-	SkeinSourceControlUtils::RunUpdateStatus(InCommand.SkeinBinaryPath, InCommand.SkeinProjectRoot, InCommand.Files, InCommand.ErrorMessages, States);
+	SkeinSourceControlUtils::RunUpdateStatus(InCommand.SkeinBinaryPath, InCommand.SkeinProjectRoot, TArray<FString>(), InCommand.Files, InCommand.ErrorMessages, States);
 
 	return InCommand.bCommandSuccessful;
 }
@@ -214,7 +217,7 @@ bool FSkeinSyncWorker::Execute(FSkeinSourceControlCommand& InCommand)
 
 	InCommand.bCommandSuccessful = SkeinSourceControlUtils::RunCommand(TEXT("projects snapshots get"), InCommand.SkeinBinaryPath, InCommand.SkeinProjectRoot, TArray<FString>(), TArray<FString>(), InCommand.InfoMessages, InCommand.ErrorMessages);
 
-	SkeinSourceControlUtils::RunUpdateStatus(InCommand.SkeinBinaryPath, InCommand.SkeinProjectRoot, InCommand.Files, InCommand.ErrorMessages, States);
+	SkeinSourceControlUtils::RunUpdateStatus(InCommand.SkeinBinaryPath, InCommand.SkeinProjectRoot, TArray<FString>(), InCommand.Files, InCommand.ErrorMessages, States);
 
 	return InCommand.bCommandSuccessful;
 }
@@ -237,7 +240,7 @@ bool FSkeinUpdateStatusWorker::Execute(FSkeinSourceControlCommand& InCommand)
 {
 	check(InCommand.Operation->GetName() == GetName());
 
-	InCommand.bCommandSuccessful = SkeinSourceControlUtils::RunUpdateStatus(InCommand.SkeinBinaryPath, InCommand.SkeinProjectRoot, InCommand.Files, InCommand.ErrorMessages, States);
+	InCommand.bCommandSuccessful = SkeinSourceControlUtils::RunUpdateStatus(InCommand.SkeinBinaryPath, InCommand.SkeinProjectRoot, TArray<FString>(), InCommand.Files, InCommand.ErrorMessages, States);
 
 	return InCommand.bCommandSuccessful;
 }
diff --git a/Engine/Plugins/Developer/SkeinSourceControl/Source/SkeinSourceControl/Private/SkeinSourceControlUtils.cpp b/Engine/Plugins/Developer/SkeinSourceControl/Source/SkeinSourceControl/Private/SkeinSourceControlUtils.cpp
--- a/Engine/Plugins/Developer/SkeinSourceControl/Source/SkeinSourceControl/Private/SkeinSourceControlUtils.cpp
+++ b/Engine/Plugins/Developer/SkeinSourceControl/Source/SkeinSourceControl/Private/SkeinSourceControlUtils.cpp
@@ -447,7 +447,7 @@ bool RunCommand(const FString& InCommand, const FString& InSkeinBinaryPath, cons
 	return RunCommandBatched(InCommand, InSkeinBinaryPath, InSkeinProjectRoot, InParameters, InFiles, Callback);
 }
 	
-bool RunUpdateStatus(const FString& InSkeinBinaryPath, const FString& InSkeinProjectRoot, const TArray<FString>& InFiles, TArray<FString>& OutErrors, TArray<FSkeinSourceControlState>& OutStates)
+bool RunUpdateStatus(const FString& InSkeinBinaryPath, const FString& InSkeinProjectRoot, const TArray<FString>& InParameters, const TArray<FString>& InFiles, TArray<FString>& OutErrors, TArray<FSkeinSourceControlState>& OutStates)
 {
 	auto Callback = 
 		[&] (bool bBatchResult, const TArray<FString>& BatchFiles, const FString& BatchMessage, const TSharedPtr<FJsonObject>& BatchData)
@@ -483,7 +483,7 @@ bool RunUpdateStatus(const FString& InSkeinBinaryPath, const FString& InSkeinPro
 		TArray<FString> Paths;
 		Paths.Add(UniquePath);
 
-		if (!RunCommandBatched(TEXT("projects status"), InSkeinBinaryPath, InSkeinProjectRoot, TArray<FString>(), Paths, Callback))
+		if (!RunCommandBatched(TEXT("projects status"), InSkeinBinaryPath, InSkeinProjectRoot, InParameters, Paths, Callback))
 		{
 			NumErrors++;
 		}
diff --git a/Engine/Plugins/Developer/SkeinSourceControl/Source/SkeinSourceControl/Private/SkeinSourceControlUtils.h b/Engine/Plugins/Developer/SkeinSourceControl/Source/SkeinSourceControl/Private/SkeinSourceControlUtils.h
--- a/Engine/Plugins/Developer/SkeinSourceControl/Source/SkeinSourceControl/Private/SkeinSourceControlUtils.h
+++ b/Engine/Plugins/Developer/SkeinSourceControl/Source/SkeinSourceControl/Private/SkeinSourceControlUtils.h
@@ -76,10 +76,11 @@ bool RunCommand(const FString& InCommand, const FString& InSkeinBinaryPath, cons
  *
  * @param	InSkeinBinaryPath	The path to the Skein binary
  * @param	InSkeinProjectRoot	The Skein project root from where to run the command
+ * @param	InParameters		The parameters to the Skein command
  * @param	InFiles				The files to be operated on
  * @param	OutErrors			Any errors as an array per-line
  * @param   OutStates           The state of each of the input files
  * @returns true if the command succeeded and returned no errors
  */
-bool RunUpdateStatus(const FString& InSkeinBinaryPath, const FString& InSkeinProjectRoot, const TArray<FString>& InFiles, TArray<FString>& OutErrors, TArray<FSkeinSourceControlState>& OutStates);
+bool RunUpdateStatus(const FString& InSkeinBinaryPath, const FString& InSkeinProjectRoot, const TArray<FString>& InParameters, const TArray<FString>& InFiles, TArray<FString>& OutErrors, TArray<FSkeinSourceControlState>& OutStates);
 }
\ No newline at end of file
