SkeinUE: Update SourceControl icon to show additional error information if there's an issue with the Skein environment (examples: binary missing, binary running, server down, logged out).

diff --git a/Engine/Plugins/Developer/SkeinSourceControl/Source/SkeinSourceControl/Private/SSkeinSourceControlSettings.cpp b/Engine/Plugins/Developer/SkeinSourceControl/Source/SkeinSourceControl/Private/SSkeinSourceControlSettings.cpp
--- a/Engine/Plugins/Developer/SkeinSourceControl/Source/SkeinSourceControl/Private/SSkeinSourceControlSettings.cpp
+++ b/Engine/Plugins/Developer/SkeinSourceControl/Source/SkeinSourceControl/Private/SSkeinSourceControlSettings.cpp
@@ -25,10 +25,10 @@ void SSkeinSourceControlSettings::Construct(const FArguments& InArgs)
 			.FillWidth(1.0f)
 			[
 				SNew(STextBlock)
-				.Text(LOCTEXT("SkeinCLINotFound", "The Skein Command Line application could not be found."))
+				.Text(LOCTEXT("SkeinCLINotFound", "The Skein environment is not available."))
 				.ColorAndOpacity(FLinearColor::Red)
 				.WrapTextAt(450.0f)
-				.ToolTipText(LOCTEXT("SkeinCLINotFound_Tooltip", "Without the Skein Command Line application the Unreal Editor cannot communicate to the Skein cloud server. Please make sure it's installed at the correct location."))
+				.ToolTipText(LOCTEXT("SkeinCLINotFound_Tooltip", "Without the Skein environment the Unreal Editor cannot communicate to the Skein cloud server. Please make sure it's installed at the correct location and that no other instances of the application are running."))
 			]
 		]
 		// Explanation text for missing Skein project
diff --git a/Engine/Plugins/Developer/SkeinSourceControl/Source/SkeinSourceControl/Private/SkeinSourceControlProvider.cpp b/Engine/Plugins/Developer/SkeinSourceControl/Source/SkeinSourceControl/Private/SkeinSourceControlProvider.cpp
--- a/Engine/Plugins/Developer/SkeinSourceControl/Source/SkeinSourceControl/Private/SkeinSourceControlProvider.cpp
+++ b/Engine/Plugins/Developer/SkeinSourceControl/Source/SkeinSourceControl/Private/SkeinSourceControlProvider.cpp
@@ -18,8 +18,11 @@
 
 void FSkeinSourceControlProvider::Init(bool bForceConnection)
 {
-	bSkeinBinaryFound = SkeinSourceControlUtils::IsSkeinAvailable();
+	bSkeinBinaryIdle = SkeinSourceControlUtils::IsSkeinBinaryIdle();
+	bSkeinBinaryFound = SkeinSourceControlUtils::IsSkeinBinaryFound();
 	bSkeinProjectFound = SkeinSourceControlUtils::IsSkeinProjectFound(FPaths::ProjectDir(), ProjectRoot, ProjectName);
+	bSkeinServerUp = true;
+	bSkeinLoggedIn = false;
 	BinaryPath = SkeinSourceControlUtils::FindSkeinBinaryPath();
 }
 
@@ -28,38 +31,56 @@ void FSkeinSourceControlProvider::Close()
 	BinaryPath.Empty();
 	ProjectName.Empty();
 	ProjectRoot.Empty();
+	bSkeinLoggedIn = false;
+	bSkeinServerUp = false;
 	bSkeinProjectFound = false;
 	bSkeinBinaryFound = false;
+	bSkeinBinaryIdle = false;
 }
 
 FText FSkeinSourceControlProvider::GetStatusText() const
 {
+	if (!bSkeinServerUp)
+	{
+		return FText(LOCTEXT("SkeinServerDown", "The Skein server is not available."));
+	}
+	if (!bSkeinBinaryFound)
+	{
+		return FText(LOCTEXT("SkeinBinaryMissing", "The 'skein' binary is missing on your system.\nPlease install it."));
+	}
+	if (!bSkeinBinaryIdle)
+	{
+		return FText(LOCTEXT("SkeinBinaryRunning", "There is already a 'skein' binary running on your system.\nPlease kill it."));
+	}
+	if (!bSkeinLoggedIn)
+	{
+		return FText(LOCTEXT("SkeinLoggedOut", "The Skein environment is not logged in.\nPlease reauthenticate."));
+	}
+
 	FFormatNamedArguments Args;
-	Args.Add(TEXT("IsAvailable"), IsAvailable() ? LOCTEXT("Yes", "Yes") : LOCTEXT("No", "No"));
-	Args.Add(TEXT("IsEnabled"), IsEnabled() ? LOCTEXT("Yes", "Yes") : LOCTEXT("No", "No"));
 	Args.Add(TEXT("ProjectName"), FText::FromString(ProjectName));
 	Args.Add(TEXT("ProjectRoot"), FText::FromString(ProjectRoot));
 
-	return FText::Format(LOCTEXT("SkeinStatusText", "Available: {IsAvailable}\nEnabled: {IsEnabled}\nProjectName: {ProjectName}\nProjectRoot: {ProjectRoot}"), Args);
+	return FText::Format(LOCTEXT("SkeinStatusText", "ProjectName: {ProjectName}\nProjectRoot: {ProjectRoot}"), Args);
 }
 
 bool FSkeinSourceControlProvider::IsAvailable() const
 {
 	if (!bSkeinBinaryFound)
 	{
-		bSkeinBinaryFound = SkeinSourceControlUtils::IsSkeinAvailable();
+		bSkeinBinaryFound = SkeinSourceControlUtils::IsSkeinBinaryFound();
+	}
+
+	if (!bSkeinBinaryIdle)
+	{
+		bSkeinBinaryIdle = SkeinSourceControlUtils::IsSkeinBinaryIdle();
 	}
 
-	return bSkeinBinaryFound;
+	return bSkeinBinaryFound && bSkeinBinaryIdle && bSkeinServerUp && bSkeinLoggedIn;
 }
 
 bool FSkeinSourceControlProvider::IsEnabled() const
 {
-	if (!bSkeinBinaryFound)
-	{
-		bSkeinBinaryFound = SkeinSourceControlUtils::IsSkeinAvailable();
-	}
-
 	if (!bSkeinProjectFound)
 	{
 		bSkeinProjectFound = SkeinSourceControlUtils::IsSkeinProjectFound(FPaths::ProjectDir(), ProjectRoot, ProjectName);
@@ -209,6 +230,20 @@ void FSkeinSourceControlProvider::Tick()
 			// Remove command from the queue
 			CommandQueue.RemoveAt(CommandIndex);
 
+			// Update binary state
+			if (Command.bCommandSuccessful)
+			{
+				bSkeinBinaryIdle = true;
+				bSkeinServerUp = true;
+				bSkeinLoggedIn = true;
+			}
+			else
+			{
+				bSkeinBinaryIdle = !Command.ErrorMessages.Contains("Another instance of 'skein' is already running.");
+				bSkeinServerUp = !Command.ErrorMessages.Contains("Skein server is not available.");
+				bSkeinLoggedIn = !Command.ErrorMessages.Contains("Skein is not logged in (or session has expired).");
+			}
+
 			// Let command update the states of any files
 			bStatesUpdated |= Command.Worker->UpdateStates();
 
diff --git a/Engine/Plugins/Developer/SkeinSourceControl/Source/SkeinSourceControl/Private/SkeinSourceControlProvider.h b/Engine/Plugins/Developer/SkeinSourceControl/Source/SkeinSourceControl/Private/SkeinSourceControlProvider.h
--- a/Engine/Plugins/Developer/SkeinSourceControl/Source/SkeinSourceControl/Private/SkeinSourceControlProvider.h
+++ b/Engine/Plugins/Developer/SkeinSourceControl/Source/SkeinSourceControl/Private/SkeinSourceControlProvider.h
@@ -92,13 +92,21 @@ private:
 	ECommandResult::Type IssueCommand(FSkeinSourceControlCommand& InCommand);
 
 private:
+	/** Is Skein.exe not obstructing us */
+	mutable bool bSkeinBinaryIdle;
 
-	/** Is Skein service available */
+	/** Is Skein.exe binary found */
 	mutable bool bSkeinBinaryFound;
 
 	/** Is Skein project found */
 	mutable bool bSkeinProjectFound;
 
+	/** Is Skein server available */
+	mutable bool bSkeinServerUp;
+
+	/** Is Skein logged in */
+	mutable bool bSkeinLoggedIn;
+
 	/** Skein project root */
 	mutable FString ProjectRoot;
 
diff --git a/Engine/Plugins/Developer/SkeinSourceControl/Source/SkeinSourceControl/Private/SkeinSourceControlUtils.cpp b/Engine/Plugins/Developer/SkeinSourceControl/Source/SkeinSourceControl/Private/SkeinSourceControlUtils.cpp
--- a/Engine/Plugins/Developer/SkeinSourceControl/Source/SkeinSourceControl/Private/SkeinSourceControlUtils.cpp
+++ b/Engine/Plugins/Developer/SkeinSourceControl/Source/SkeinSourceControl/Private/SkeinSourceControlUtils.cpp
@@ -343,7 +343,7 @@ FString GetIntermediateThumbnailPath(const FString& InPath, const FString& InInt
 	return FString();
 }
 
-bool IsSkeinAvailable()
+bool IsSkeinBinaryFound()
 {
 	FString SkeinBinaryPath = FindSkeinBinaryPath();
 	if (SkeinBinaryPath.IsEmpty())
@@ -352,6 +352,17 @@ bool IsSkeinAvailable()
 	return true;
 }
 
+bool IsSkeinBinaryIdle()
+{
+	FString SkeinBinaryPath = FindSkeinBinaryPath();
+	if (SkeinBinaryPath.IsEmpty())
+		return false;
+
+	FString SkeinBinaryFile = FPaths::GetCleanFilename(SkeinBinaryPath);
+
+	return !FPlatformProcess::IsApplicationRunning(*SkeinBinaryFile);
+}
+
 bool IsSkeinProjectFound(const FString& InPath, FString& OutProjectRoot, FString& OutProjectName)
 {
 	OutProjectRoot = FindSkeinProjectRoot(InPath);
diff --git a/Engine/Plugins/Developer/SkeinSourceControl/Source/SkeinSourceControl/Private/SkeinSourceControlUtils.h b/Engine/Plugins/Developer/SkeinSourceControl/Source/SkeinSourceControl/Private/SkeinSourceControlUtils.h
--- a/Engine/Plugins/Developer/SkeinSourceControl/Source/SkeinSourceControl/Private/SkeinSourceControlUtils.h
+++ b/Engine/Plugins/Developer/SkeinSourceControl/Source/SkeinSourceControl/Private/SkeinSourceControlUtils.h
@@ -41,11 +41,16 @@ FString GetIntermediateMetadataPath(const FString& InPath, const FString& InInte
 FString GetIntermediateThumbnailPath(const FString& InPath, const FString& InIntermediateRoot = FString());
 
 /**
- * Finds out if the Skein environment is available.
- * In practice this checks if the Skein CLI application is installed at the expected location.
- * @returns true if the Skein environment can be used for source control operations.
+ * Finds out if the Skein application is installed.
+  * @returns true if the Skein application is installed.
  */
-bool IsSkeinAvailable();
+bool IsSkeinBinaryFound();
+
+/**
+ * Finds out if the Skein application is running or not.
+ * @returns true if the Skein application is idle.
+ */
+bool IsSkeinBinaryIdle();
 
 /**
   * Finds out if the given directory is part of a Skein project
