SkeinUE: Fixed error message incorrectly popping up in 'Connect to Source Control' dialog.

When initially connecting, the user is not authenticated yet. As this is checked by SkeinSourceControlProvider::IsAvailable that function cannot be called in this dialog as it will always return 'false'.

diff --git a/Engine/Plugins/Developer/SkeinSourceControl/Source/SkeinSourceControl/Private/SSkeinSourceControlSettings.cpp b/Engine/Plugins/Developer/SkeinSourceControl/Source/SkeinSourceControl/Private/SSkeinSourceControlSettings.cpp
--- a/Engine/Plugins/Developer/SkeinSourceControl/Source/SkeinSourceControl/Private/SSkeinSourceControlSettings.cpp
+++ b/Engine/Plugins/Developer/SkeinSourceControl/Source/SkeinSourceControl/Private/SSkeinSourceControlSettings.cpp
@@ -2,6 +2,7 @@
 
 #include "SSkeinSourceControlSettings.h"
 #include "SkeinSourceControlModule.h"
+#include "SkeinSourceControlUtils.h"
 #include "Modules/ModuleManager.h"
 #include "Widgets/SBoxPanel.h"
 #include "Widgets/Text/STextBlock.h"
@@ -55,16 +56,16 @@ void SSkeinSourceControlSettings::Construct(const FArguments& InArgs)
 EVisibility SSkeinSourceControlSettings::CanUseSkeinCLI() const
 {
 	FSkeinSourceControlModule& SkeinSourceControl = FModuleManager::LoadModuleChecked<FSkeinSourceControlModule>("SkeinSourceControl");
-	const bool bSkeinAvailable = SkeinSourceControl.GetProvider().IsAvailable();
-	return (!bSkeinAvailable) ? EVisibility::Visible : EVisibility::Collapsed;
+	const bool bSkeinBinaryFound = SkeinSourceControlUtils::IsSkeinBinaryFound();
+	return (!bSkeinBinaryFound) ? EVisibility::Visible : EVisibility::Collapsed;
 }
 
 EVisibility SSkeinSourceControlSettings::CanUseSkeinProject() const
 {
 	FSkeinSourceControlModule& SkeinSourceControl = FModuleManager::LoadModuleChecked<FSkeinSourceControlModule>("SkeinSourceControl");
-	const bool bSkeinAvailable = SkeinSourceControl.GetProvider().IsAvailable();
-	const bool bSkeinEnabled = SkeinSourceControl.GetProvider().IsEnabled();
-	return (bSkeinAvailable && !bSkeinEnabled) ? EVisibility::Visible : EVisibility::Collapsed;
+	const bool bSkeinBinaryFound = SkeinSourceControlUtils::IsSkeinBinaryFound();
+	const bool bSkeinProjectFound = SkeinSourceControlUtils::IsSkeinProjectFound(FPaths::ProjectDir());
+	return (bSkeinBinaryFound && !bSkeinProjectFound) ? EVisibility::Visible : EVisibility::Collapsed;
 }
 
 #undef LOCTEXT_NAMESPACE
diff --git a/Engine/Plugins/Developer/SkeinSourceControl/Source/SkeinSourceControl/Private/SkeinSourceControlUtils.cpp b/Engine/Plugins/Developer/SkeinSourceControl/Source/SkeinSourceControl/Private/SkeinSourceControlUtils.cpp
--- a/Engine/Plugins/Developer/SkeinSourceControl/Source/SkeinSourceControl/Private/SkeinSourceControlUtils.cpp
+++ b/Engine/Plugins/Developer/SkeinSourceControl/Source/SkeinSourceControl/Private/SkeinSourceControlUtils.cpp
@@ -363,6 +363,13 @@ bool IsSkeinBinaryIdle()
 	return !FPlatformProcess::IsApplicationRunning(*SkeinBinaryFile);
 }
 
+bool IsSkeinProjectFound(const FString& InPath)
+{
+	FString ProjectRootDummy;
+	FString ProjectNameDummy;
+	return IsSkeinProjectFound(InPath, ProjectRootDummy, ProjectNameDummy);
+}
+
 bool IsSkeinProjectFound(const FString& InPath, FString& OutProjectRoot, FString& OutProjectName)
 {
 	OutProjectRoot = FindSkeinProjectRoot(InPath);
diff --git a/Engine/Plugins/Developer/SkeinSourceControl/Source/SkeinSourceControl/Private/SkeinSourceControlUtils.h b/Engine/Plugins/Developer/SkeinSourceControl/Source/SkeinSourceControl/Private/SkeinSourceControlUtils.h
--- a/Engine/Plugins/Developer/SkeinSourceControl/Source/SkeinSourceControl/Private/SkeinSourceControlUtils.h
+++ b/Engine/Plugins/Developer/SkeinSourceControl/Source/SkeinSourceControl/Private/SkeinSourceControlUtils.h
@@ -60,6 +60,7 @@ bool IsSkeinBinaryIdle();
   * @param	 OutProjectName		The Skein project name found (if any)
   * @returns true if a Skein project was found.
   */
+bool IsSkeinProjectFound(const FString& InPath);
 bool IsSkeinProjectFound(const FString& InPath, FString& OutProjectRoot, FString& OutProjectName);
 
 /**
