SkeinUE: Reformat dependency names submitted to SkeinAPI to the project root relative asset name format used by Skein and its WebUI.

diff --git a/Engine/Plugins/Developer/SkeinSourceControl/Source/SkeinSourceControl/Private/SkeinSourceControlMetadata.cpp b/Engine/Plugins/Developer/SkeinSourceControl/Source/SkeinSourceControl/Private/SkeinSourceControlMetadata.cpp
--- a/Engine/Plugins/Developer/SkeinSourceControl/Source/SkeinSourceControl/Private/SkeinSourceControlMetadata.cpp
+++ b/Engine/Plugins/Developer/SkeinSourceControl/Source/SkeinSourceControl/Private/SkeinSourceControlMetadata.cpp
@@ -1,9 +1,11 @@
 // Copyright Epic Games, Inc. All Rights Reserved.
 
 #include "SkeinSourceControlMetadata.h"
+#include "SkeinSourceControlModule.h"
 #include "AssetRegistry/AssetRegistryModule.h"
 #include "AssetRegistry/IAssetRegistry.h"
 #include "Serialization/JsonSerializer.h"
+#include "Modules/ModuleManager.h"
 #include "Misc/AssetRegistryInterface.h"
 #include "Misc/FileHelper.h"
 #include "Algo/Transform.h"
@@ -57,9 +59,59 @@ static bool ExtractTags(IAssetRegistry& InAssetRegistry, FAssetData& InAssetData
 }
 
 // Dependencies are straightforward. The AssetRegistry provides a nice interface for them.
-static bool ExtractDependencies(IAssetRegistry& InAssetRegistry, FAssetData& InAssetData, TArray<FName>& OutDependencies)
+// We do reformat to the project root relative asset name format that's being used by Skein and its WebUI.
+// For example:
+//   Dep In : /Game/Environment/Textures/T_Vista_Grass_N
+//   Dep Out: /Content/Environment/Textures/T_Vista_Grass_N.uasset
+static bool ExtractDependencies(IAssetRegistry& InAssetRegistry, FAssetData& InAssetData, TArray<FString>& OutDependencies)
 {
-	return InAssetRegistry.GetDependencies(InAssetData.PackageName, OutDependencies);
+	FSkeinSourceControlModule& SkeinSourceControl = FModuleManager::LoadModuleChecked<FSkeinSourceControlModule>("SkeinSourceControl");
+	FString SkeinProjectRoot = SkeinSourceControl.GetProvider().GetSkeinProjectRoot() + TEXT("/");
+
+	TArray<FName> Dependencies;
+	if (InAssetRegistry.GetDependencies(InAssetData.PackageName, Dependencies))
+	{
+		for (const FName& Dependency : Dependencies)
+		{
+			FString RelativeDepFileName;
+			if (FPackageName::TryConvertLongPackageNameToFilename(Dependency.ToString(), RelativeDepFileName))
+			{
+				FString DepFileName;
+
+				if (DepFileName.IsEmpty())
+				{
+					FString MapFileName = RelativeDepFileName + FPackageName::GetMapPackageExtension();
+					if (IFileManager::Get().FileExists(*MapFileName))
+					{
+						DepFileName = FPaths::ConvertRelativePathToFull(MapFileName);
+					}
+				}
+
+				if (DepFileName.IsEmpty())
+				{
+					FString AssetFileName = RelativeDepFileName + FPackageName::GetAssetPackageExtension();
+					if (IFileManager::Get().FileExists(*AssetFileName))
+					{
+						DepFileName = FPaths::ConvertRelativePathToFull(AssetFileName);
+					}
+				}
+
+				if (!DepFileName.IsEmpty())
+				{
+					if (FPaths::IsUnderDirectory(DepFileName, SkeinProjectRoot))
+					{
+						if (FPaths::MakePathRelativeTo(DepFileName, *SkeinProjectRoot))
+						{
+							OutDependencies.Add(DepFileName);
+						}
+					}
+				}
+			}
+		}
+
+		return true;
+	}
+	return false;
 }
 
 // ThumbnailTools provides some nice features for us. The Thumbnail could be missing though.
@@ -143,7 +195,7 @@ static bool WriteThumbnailToDisk(const FString& InThumbnailPath, const FObjectTh
 	return true;
 }
 
-static bool WriteMetadataToDisk(const FString& InMetadataPath, const TMap<FName, FString>& InTags, const TArray<FName>& InDependencies)
+static bool WriteMetadataToDisk(const FString& InMetadataPath, const TMap<FName, FString>& InTags, const TArray<FString>& InDependencies)
 {
 	// Build JsonObject
 	TArray<TSharedPtr<FJsonValue>> JsonTagsArray;
@@ -158,9 +210,9 @@ static bool WriteMetadataToDisk(const FString& InMetadataPath, const TMap<FName,
 
 	TArray<TSharedPtr<FJsonValue>> JsonDepsArray;
 	JsonDepsArray.Reserve(InDependencies.Num());
-	for (const FName& Dependency : InDependencies)
+	for (const FString& Dependency : InDependencies)
 	{
-		JsonDepsArray.Add(MakeShareable(new FJsonValueString(Dependency.ToString())));
+		JsonDepsArray.Add(MakeShareable(new FJsonValueString(Dependency)));
 	}
 
 	TSharedRef<FJsonObject> JsonData = MakeShareable(new FJsonObject);
@@ -217,7 +269,7 @@ bool ExtractMetadata(const FString& InPackagePath, const FString& InMetadataPath
 	TMap<FName, FString> Tags;
 	ExtractTags(AssetRegistry, AssetData, Tags);
 
-	TArray<FName> Dependencies;
+	TArray<FString> Dependencies;
 	ExtractDependencies(AssetRegistry, AssetData, Dependencies);
 
 	FObjectThumbnail Thumbnail;
