SkeinUE: Fixed assert that hit when calling GetAssetRegistryTags on an UWorld from another thread. This means no tags will be on the WebUI for World assets for now.

diff --git a/Engine/Plugins/Developer/SkeinSourceControl/Source/SkeinSourceControl/Private/SkeinSourceControlMetadata.cpp b/Engine/Plugins/Developer/SkeinSourceControl/Source/SkeinSourceControl/Private/SkeinSourceControlMetadata.cpp
--- a/Engine/Plugins/Developer/SkeinSourceControl/Source/SkeinSourceControl/Private/SkeinSourceControlMetadata.cpp
+++ b/Engine/Plugins/Developer/SkeinSourceControl/Source/SkeinSourceControl/Private/SkeinSourceControlMetadata.cpp
@@ -30,17 +30,22 @@ static bool ExtractTags(IAssetRegistry& InAssetRegistry, FAssetData& InAssetData
 
 		UObject* Asset = InAssetData.GetAsset();
 
-		TArray<UObject::FAssetRegistryTag> AssetRegistryTags;
-		Asset->GetAssetRegistryTags(AssetRegistryTags);
-
-		for (TArray<UObject::FAssetRegistryTag>::TConstIterator TagIter(AssetRegistryTags); TagIter; ++TagIter)
+		// UWorld::GetAssetRegistryTags broadcasts and one of the registered delegates asserts on
+		// a game thread requirement in FActorIteratorState.
+		static FName NAME_World("World");
+		if (Asset->GetClass()->GetFName() != NAME_World)
 		{
-			OutTags.Add(TagIter->Name, TagIter->Value);
+			TArray<UObject::FAssetRegistryTag> AssetRegistryTags;
+			Asset->GetAssetRegistryTags(AssetRegistryTags);
+
+			for (TArray<UObject::FAssetRegistryTag>::TConstIterator TagIter(AssetRegistryTags); TagIter; ++TagIter)
+			{
+				OutTags.Add(TagIter->Name, TagIter->Value);
+			}
 		}
 	}
 	else
 	{
-
 		auto IsValid = [](const TTuple<FName, FString>& InTuple)
 		{
 			static FName NAME_FiBData("FiBData");
