SkeinUE: Increase files per batch, ensure unique output and invoke 'projects status' using multiple paths at once. This massively speeds up source control operations on multiple directories.

diff --git a/Engine/Plugins/Developer/SkeinSourceControl/Source/SkeinSourceControl/Private/SkeinSourceControlUtils.cpp b/Engine/Plugins/Developer/SkeinSourceControl/Source/SkeinSourceControl/Private/SkeinSourceControlUtils.cpp
--- a/Engine/Plugins/Developer/SkeinSourceControl/Source/SkeinSourceControl/Private/SkeinSourceControlUtils.cpp
+++ b/Engine/Plugins/Developer/SkeinSourceControl/Source/SkeinSourceControl/Private/SkeinSourceControlUtils.cpp
@@ -1,6 +1,7 @@
 // Copyright Epic Games, Inc. All Rights Reserved.
 
 #include "SkeinSourceControlUtils.h"
+#include "ISourceControlModule.h"
 #include "HAL/PlatformProcess.h"
 #include "HAL/PlatformFileManager.h"
 #include "HAL/FileManager.h"
@@ -8,21 +9,30 @@
 #include "Misc/Paths.h"
 #include "Misc/SecureHash.h"
 #include "Misc/ScopeLock.h"
-#include "ISourceControlModule.h"
 #include "Serialization/JsonSerializer.h"
 #include "Dom/JsonValue.h"
 
 namespace SkeinSourceControlConstants
 {
 	/** The maximum number of files we submit in a single Skein command */
-	const int32 MaxFilesPerBatch = 32;
+	const int32 MaxFilesPerBatch = 50;
 }
 
 namespace SkeinSourceControlUtils
 {
 
+// A set of FSkeinSourceControlState's that only contain one entry for each Filename.
+static struct FSkeinSourceControlStateKeyFuncs : BaseKeyFuncs<FSkeinSourceControlState, FString>
+{
+	static const FString& GetSetKey(const FSkeinSourceControlState& Element) { return Element.Filename; }
+	static bool Matches(const FString& A, const FString& B) { return A == B; }
+	static uint32 GetKeyHash(const FString& Key) { return GetTypeHash(Key); }
+};
+
+typedef TSet<FSkeinSourceControlState, FSkeinSourceControlStateKeyFuncs> FSkeinSourceControlStateSet;
+
 // Convert an array of FJsonObject statuses to FSkeinSourceControlStates
-static void ParseStatusOutput(const FString& InSkeinProjectRoot, const TArray<FString>& InFiles, const TSharedPtr<FJsonObject> InStates, TArray<FSkeinSourceControlState>& OutStates)
+static void ParseStatusOutput(const FString& InSkeinProjectRoot, const TArray<FString>& InFiles, const TSharedPtr<FJsonObject> InStates, FSkeinSourceControlStateSet& OutStates)
 {
 	const FDateTime Now = FDateTime::Now();
 
@@ -467,12 +477,14 @@ bool RunCommand(const FString& InCommand, const FString& InSkeinBinaryPath, cons
 	
 bool RunUpdateStatus(const FString& InSkeinBinaryPath, const FString& InSkeinProjectRoot, const TArray<FString>& InParameters, const TArray<FString>& InFiles, TArray<FString>& OutErrors, TArray<FSkeinSourceControlState>& OutStates)
 {
+	FSkeinSourceControlStateSet States;
+
 	auto Callback = 
 		[&] (bool bBatchResult, const TArray<FString>& BatchFiles, const FString& BatchMessage, const TSharedPtr<FJsonObject>& BatchData)
 		{
 			if (bBatchResult)
 			{
-				ParseStatusOutput(InSkeinProjectRoot, InFiles, BatchData, OutStates);
+				ParseStatusOutput(InSkeinProjectRoot, InFiles, BatchData, States);
 			}
 			else
 			{
@@ -480,34 +492,39 @@ bool RunUpdateStatus(const FString& InSkeinBinaryPath, const FString& InSkeinPro
 			}
 		};
 
-	TSet<FString> UniquePaths;
+	TArray<FString> UniquePaths;
 	for (const auto& File : InFiles)
 	{
 		const bool bIsDirectory = IFileManager::Get().DirectoryExists(*File);
 		if (bIsDirectory)
 		{
-			UniquePaths.Add(*File);
+			UniquePaths.AddUnique(*File);
 		}
 		else
 		{
-			UniquePaths.Add(FPaths::GetPath(*File));
+			UniquePaths.AddUnique(FPaths::GetPath(*File));
 		}
 	}
 
-	int NumErrors = 0;
+	bool bResult = RunCommandBatched(TEXT("projects status"), InSkeinBinaryPath, InSkeinProjectRoot, InParameters, UniquePaths, Callback);
 
-	for (const auto& UniquePath : UniquePaths)
+	if (bResult)
 	{
-		TArray<FString> Paths;
-		Paths.Add(UniquePath);
-
-		if (!RunCommandBatched(TEXT("projects status"), InSkeinBinaryPath, InSkeinProjectRoot, InParameters, Paths, Callback))
-		{
-			NumErrors++;
-		}
+		// The UniquePaths array can contain a hierarchy of paths, for example:
+		// * /Path/To/A/B/
+		// * /Path/To/A/
+		//
+		// A file that's already under our control at the location /Path/To/A/B/file.uasset will be
+		// in the output of both these commands. Untracked files however do not appear for nested
+		// directories so all paths need to be evaluated individually.
+		//
+		// To remove duplicates a TSet with the FileName as the key was used, which is now
+		// converted to the output type.
+
+		OutStates = States.Array();
 	}
 
-	return (NumErrors == 0);
+	return true;
 }
 	
 }
\ No newline at end of file
