SkeinUE: Fix compile error on Mac (this time actually saving the file... doh!)

diff --git a/Engine/Plugins/Developer/SkeinSourceControl/Source/SkeinSourceControl/Private/SkeinSourceControlUtils.cpp b/Engine/Plugins/Developer/SkeinSourceControl/Source/SkeinSourceControl/Private/SkeinSourceControlUtils.cpp
--- a/Engine/Plugins/Developer/SkeinSourceControl/Source/SkeinSourceControl/Private/SkeinSourceControlUtils.cpp
+++ b/Engine/Plugins/Developer/SkeinSourceControl/Source/SkeinSourceControl/Private/SkeinSourceControlUtils.cpp
@@ -22,7 +22,7 @@ namespace SkeinSourceControlUtils
 {
 
 // A set of FSkeinSourceControlState's that only contain one entry for each Filename.
-static struct FSkeinSourceControlStateKeyFuncs : BaseKeyFuncs<FSkeinSourceControlState, FString>
+struct FSkeinSourceControlStateKeyFuncs : BaseKeyFuncs<FSkeinSourceControlState, FString>
 {
 	static const FString& GetSetKey(const FSkeinSourceControlState& Element) { return Element.Filename; }
 	static bool Matches(const FString& A, const FString& B) { return A == B; }
