SkeinUE: Filter asset tags in both code paths.

diff --git a/Engine/Plugins/Developer/SkeinSourceControl/Source/SkeinSourceControl/Private/SkeinSourceControlMetadata.cpp b/Engine/Plugins/Developer/SkeinSourceControl/Source/SkeinSourceControl/Private/SkeinSourceControlMetadata.cpp
--- a/Engine/Plugins/Developer/SkeinSourceControl/Source/SkeinSourceControl/Private/SkeinSourceControlMetadata.cpp
+++ b/Engine/Plugins/Developer/SkeinSourceControl/Source/SkeinSourceControl/Private/SkeinSourceControlMetadata.cpp
@@ -5,6 +5,7 @@
 #include "AssetRegistry/AssetRegistryModule.h"
 #include "AssetRegistry/IAssetRegistry.h"
 #include "Serialization/JsonSerializer.h"
+#include "Policies/CondensedJsonPrintPolicy.h"
 #include "Modules/ModuleManager.h"
 #include "Misc/AssetRegistryInterface.h"
 #include "Misc/FileHelper.h"
@@ -22,6 +23,17 @@ namespace SkeinSourceControlMetadata
 // is loaded.
 static bool ExtractTags(IAssetRegistry& InAssetRegistry, FAssetData& InAssetData, TMap<FName, FString>& OutTags)
 {
+	auto IsValid = [](const TTuple<FName, FString>& InTuple)
+	{
+		static FName NAME_FiBData("FiBData");
+		static FName NAME_ClassFlags("ClassFlags");
+		static FName NAME_AssetImportData("AssetImportData");
+
+		return InTuple.Key != NAME_FiBData
+			&& InTuple.Key != NAME_ClassFlags
+			&& InTuple.Key != NAME_AssetImportData;
+	};
+
 	if (InAssetData.IsAssetLoaded())
 	{
 		// UObject::GetAssetRegistryTags might access render resources so we need to tag this thread
@@ -40,23 +52,15 @@ static bool ExtractTags(IAssetRegistry& InAssetRegistry, FAssetData& InAssetData
 
 			for (TArray<UObject::FAssetRegistryTag>::TConstIterator TagIter(AssetRegistryTags); TagIter; ++TagIter)
 			{
-				OutTags.Add(TagIter->Name, TagIter->Value);
+				if (IsValid(MakeTuple(TagIter->Name, TagIter->Value)))
+				{
+					OutTags.Add(TagIter->Name, TagIter->Value);
+				}
 			}
 		}
 	}
 	else
 	{
-		auto IsValid = [](const TTuple<FName, FString>& InTuple)
-		{
-			static FName NAME_FiBData("FiBData");
-			static FName NAME_ClassFlags("ClassFlags");
-			static FName NAME_AssetImportData("AssetImportData");
-
-			return InTuple.Key != NAME_FiBData
-				&& InTuple.Key != NAME_ClassFlags
-				&& InTuple.Key != NAME_AssetImportData;
-		};
-
 		Algo::TransformIf(InAssetData.TagsAndValues.CopyMap(), OutTags, IsValid, [](const TTuple<FName, FString>& InTuple) { return InTuple; });
 	}
 
@@ -225,8 +229,8 @@ static bool WriteMetadataToDisk(const FString& InMetadataPath, const TMap<FName,
 	JsonData->SetArrayField("deps", JsonDepsArray);
 
 	// Write Json to disk
-	FString JsonString;
-	TSharedRef<TJsonWriter<>> JsonWriter = TJsonWriterFactory<>::Create(&JsonString);
+	auto JsonString = FString();
+	auto JsonWriter = TJsonWriterFactory<TCHAR, TCondensedJsonPrintPolicy<TCHAR>>::Create(&JsonString);
 	if (!FJsonSerializer::Serialize(JsonData, JsonWriter))
 	{
 #if UE_BUILD_DEBUG
